[FloatingProperties] Moved plugin

diff --git a/Engine/Plugins/Experimental/FloatingProperties/FloatingProperties.uplugin b/Engine/Plugins/Experimental/FloatingProperties/FloatingProperties.uplugin
new file mode 100644
--- /dev/null
+++ b/Engine/Plugins/Experimental/FloatingProperties/FloatingProperties.uplugin
@@ -0,0 +1,23 @@
+{
+	"FileVersion": 3,
+	"Version": 0.1,
+	"VersionName": "0.1.0-alpha",
+	"FriendlyName": "Floating Properties",
+	"Description": "Show floating properties from the details panel directly on the active viewport.",
+	"Category": "Experimental",
+	"CreatedBy": "Epic Games, Inc.",
+	"CreatedByURL": "http://epicgames.com",
+	"DocsURL": "",
+	"MarketplaceURL": "",
+	"SupportURL": "",
+	"CanContainContent": true,
+	"IsBetaVersion": true,
+	"Installed": false,
+	"Modules": [
+		{
+			"Name": "FloatingProperties",
+			"Type": "Editor",
+			"LoadingPhase": "Default"
+		}
+	]
+}
diff --git a/Engine/Plugins/Experimental/FloatingProperties/Source/FloatingProperties/FloatingProperties.Build.cs b/Engine/Plugins/Experimental/FloatingProperties/Source/FloatingProperties/FloatingProperties.Build.cs
new file mode 100644
--- /dev/null
+++ b/Engine/Plugins/Experimental/FloatingProperties/Source/FloatingProperties/FloatingProperties.Build.cs
@@ -0,0 +1,36 @@
+// Copyright Epic Games, Inc. All Rights Reserved.
+
+using UnrealBuildTool;
+
+public class FloatingProperties : ModuleRules
+{
+	public FloatingProperties(ReadOnlyTargetRules Target) : base(Target)
+	{
+		PCHUsage = ModuleRules.PCHUsageMode.UseExplicitOrSharedPCHs;
+
+		PublicDependencyModuleNames.AddRange(
+			new string[]
+			{
+				"Core",
+				"DeveloperSettings"
+			}
+		);
+
+		PrivateDependencyModuleNames.AddRange(
+			new string[]
+			{
+				"AppFramework",
+				"CoreUObject",
+				"EditorFramework",
+				"Engine",
+				"InputCore",
+				"Projects",
+				"PropertyEditor",
+				"Slate",
+				"SlateCore",
+				"ToolMenus",
+				"UnrealEd"
+			}
+		);
+	}
+}
diff --git a/Engine/Plugins/Experimental/FloatingProperties/Source/FloatingProperties/Private/Data/FloatingPropertiesPropertyNode.cpp b/Engine/Plugins/Experimental/FloatingProperties/Source/FloatingProperties/Private/Data/FloatingPropertiesPropertyNode.cpp
new file mode 100644
--- /dev/null
+++ b/Engine/Plugins/Experimental/FloatingProperties/Source/FloatingProperties/Private/Data/FloatingPropertiesPropertyNode.cpp
@@ -0,0 +1,542 @@
+// Copyright Epic Games, Inc. All Rights Reserved.
+
+#include "Data/FloatingPropertiesPropertyNode.h"
+#include "Data/FloatingPropertiesPropertyNodeContainer.h"
+#include "FloatingPropertiesModule.h"
+#include "FloatingPropertiesSettings.h"
+#include "Math/Vector2D.h"
+#include "Types/SlateEnums.h"
+#include "Widgets/SFloatingPropertiesPropertyWidget.h"
+#include "Widgets/SWidget.h"
+
+EFloatingPropertiesUpdateResult UE::FloatingProperties::CombineUpdateResult(EFloatingPropertiesUpdateResult InCurrentResult, 
+	EFloatingPropertiesUpdateResult InNewResult)
+{
+	if (InNewResult == EFloatingPropertiesUpdateResult::Failed)
+	{
+		return EFloatingPropertiesUpdateResult::Failed;
+	}
+
+	switch (InCurrentResult)
+	{
+		case EFloatingPropertiesUpdateResult::Failed:
+			return EFloatingPropertiesUpdateResult::Failed;
+
+		case EFloatingPropertiesUpdateResult::Updated:
+			return EFloatingPropertiesUpdateResult::Updated;
+
+		case EFloatingPropertiesUpdateResult::AlreadyUpToDate:
+		default:
+			return InNewResult;
+	}
+}
+
+FFloatingPropertiesPropertyNode::FFloatingPropertiesPropertyNode(FloatingPropertiesPropertyNodeContainer* InLinkedList,
+	TSharedRef<SFloatingPropertiesPropertyWidget> InPropertyWidget, int32 InIndex)
+	: LinkedList(InLinkedList)
+    , MyIndex(InIndex)
+    , ParentNodeIndex(INDEX_NONE)
+    , ChildNodeIndex(INDEX_NONE)
+	, PropertyWidget(InPropertyWidget)
+	, PropertyAnchor(FFloatingPropertiesClassPropertyAnchor())
+	, PropertyPosition(FFloatingPropertiesClassPropertyPosition::DefaultStackPosition)
+    , CachedPosition(FVector2f::ZeroVector)
+    , bValidPosition(false)
+{		
+}
+
+void FFloatingPropertiesPropertyNode::SetChild(TSharedRef<FFloatingPropertiesPropertyNode> InChild)
+{
+	if (InChild->MyIndex == MyIndex
+		|| InChild->ParentNodeIndex == MyIndex
+		|| InChild->ChildNodeIndex == MyIndex
+		|| InChild->MyIndex == ParentNodeIndex
+		|| InChild->MyIndex == ChildNodeIndex)
+	{
+		return;
+	}
+
+	if (TSharedPtr<FFloatingPropertiesPropertyNode> Child = GetChild())
+	{
+		if (Child->MyIndex == InChild->MyIndex)
+		{
+			return;
+		}
+
+		RemoveChild();
+	}
+
+	InChild->RemoveParent();	
+
+	ChildNodeIndex = InChild->MyIndex;
+	PropertyAnchor.ChildProperty = InChild->PropertyWidget->GetClassProperty();
+
+	InChild->ParentNodeIndex = MyIndex;
+	InChild->PropertyAnchor.ParentProperty = PropertyWidget->GetClassProperty();
+	InChild->UpdateStackPositions(/* bInInvalidate */ true);
+}
+
+void FFloatingPropertiesPropertyNode::SetParent(TSharedPtr<FFloatingPropertiesPropertyNode> InParentNode)
+{
+	InParentNode->SetChild(SharedThis(this));
+}
+
+TSharedPtr<FFloatingPropertiesPropertyNode> FFloatingPropertiesPropertyNode::RemoveChild()
+{
+	TSharedPtr<FFloatingPropertiesPropertyNode> Child = GetChild();
+
+	if (!Child.IsValid())
+	{
+		return nullptr;
+	}
+
+	ChildNodeIndex = INDEX_NONE;
+	PropertyAnchor.ChildProperty.Reset();
+
+	Child->ParentNodeIndex = INDEX_NONE;
+	Child->PropertyAnchor.ParentProperty.Reset();
+
+	Child->PropertyPosition = CalculatePropertyPosition(
+		LinkedList->GetDraggableArea(), 
+		Child->GetPropertyWidget()->GetDesiredSize(),
+		Child->GetCachedPosition()
+	);
+
+	Child->UpdateStackPositions(/* bInInvalidate */ true);
+
+	return Child;
+}
+
+TSharedPtr<FFloatingPropertiesPropertyNode> FFloatingPropertiesPropertyNode::RemoveParent()
+{
+	if (TSharedPtr<FFloatingPropertiesPropertyNode> Parent = GetParent())
+	{
+		Parent->RemoveChild();
+		return Parent;
+	}
+
+	return nullptr;
+}
+
+void FFloatingPropertiesPropertyNode::InsertAfter(TSharedRef<FFloatingPropertiesPropertyNode> InMiddle)
+{
+	RemoveFromStack();
+
+	TSharedPtr<FFloatingPropertiesPropertyNode> OldChild = InMiddle->RemoveChild();
+
+	InMiddle->SetChild(SharedThis(this));
+
+	if (OldChild.IsValid())
+	{
+		TSharedRef<FFloatingPropertiesPropertyNode> Leaf = GetStackLeafMostNode();
+		Leaf->SetChild(OldChild.ToSharedRef());
+	}	
+}
+
+void FFloatingPropertiesPropertyNode::RemoveFromStack()
+{
+	TSharedPtr<FFloatingPropertiesPropertyNode> Child = RemoveChild();
+	TSharedPtr<FFloatingPropertiesPropertyNode> Parent = RemoveParent();
+
+	if (Parent.IsValid() && Child.IsValid())
+	{
+		Parent->SetChild(Child.ToSharedRef());
+	}
+}
+
+bool FFloatingPropertiesPropertyNode::HasParent() const
+{
+	return ParentNodeIndex != MyIndex && LinkedList->GetNodes().IsValidIndex(ParentNodeIndex);
+}
+
+bool FFloatingPropertiesPropertyNode::HasParent(TSharedRef<FFloatingPropertiesPropertyNode> InParent) const
+{
+	TSharedPtr<const FFloatingPropertiesPropertyNode> LoopNode = SharedThis(this);
+
+	while (LoopNode.IsValid())
+	{
+		if (LoopNode == InParent)
+		{
+			return true;
+		}
+
+		LoopNode = LoopNode->GetParent();
+	}
+
+	return false;
+}
+
+TSharedPtr<FFloatingPropertiesPropertyNode> FFloatingPropertiesPropertyNode::GetParent() const
+{
+	// Error detection
+	if (ParentNodeIndex == MyIndex)
+	{
+		UE_LOG(LogFloatingProperties, Warning, TEXT("Invalid parent relationship detected in GetParent."));
+		return nullptr;
+	}
+
+	if (LinkedList->GetNodes().IsValidIndex(ParentNodeIndex))
+	{
+		return LinkedList->GetNodes()[ParentNodeIndex];
+	}
+
+	return nullptr;
+}
+
+bool FFloatingPropertiesPropertyNode::HasChild() const
+{
+	return ChildNodeIndex != MyIndex && LinkedList->GetNodes().IsValidIndex(ChildNodeIndex);
+}
+
+bool FFloatingPropertiesPropertyNode::HasChild(TSharedRef<FFloatingPropertiesPropertyNode> InChild) const
+{
+	return InChild->HasParent(SharedThis(const_cast<FFloatingPropertiesPropertyNode*>(this)));
+}
+
+TSharedPtr<FFloatingPropertiesPropertyNode> FFloatingPropertiesPropertyNode::GetChild() const
+{
+	// Error detection
+	if (ChildNodeIndex == MyIndex)
+	{
+		UE_LOG(LogFloatingProperties, Warning, TEXT("Invalid parent relationship detected in GetChild."));
+		return nullptr;
+	}
+
+	if (LinkedList->GetNodes().IsValidIndex(ChildNodeIndex))
+	{
+		return LinkedList->GetNodes()[ChildNodeIndex];
+	}
+
+	return nullptr;
+}
+
+bool FFloatingPropertiesPropertyNode::HasSibling(TSharedRef<FFloatingPropertiesPropertyNode> InSibling) const
+{
+	TSharedRef<FFloatingPropertiesPropertyNode> This = SharedThis(const_cast<FFloatingPropertiesPropertyNode*>(this));
+
+	if (This == InSibling)
+	{
+		return true;
+	}
+
+	TSharedPtr<FFloatingPropertiesPropertyNode> LoopNode = This;
+
+	while (TSharedPtr<FFloatingPropertiesPropertyNode> Parent = LoopNode->GetParent())
+	{
+		LoopNode = Parent.ToSharedRef();
+
+		if (LoopNode == InSibling)
+		{
+			return true;
+		}
+	}
+
+	LoopNode = This;
+
+	while (TSharedPtr<FFloatingPropertiesPropertyNode> Child = LoopNode->GetChild())
+	{
+		LoopNode = Child.ToSharedRef();
+
+		if (LoopNode == InSibling)
+		{
+			return true;
+		}
+	}
+
+	return false;
+}
+
+TSharedRef<FFloatingPropertiesPropertyNode> FFloatingPropertiesPropertyNode::GetStackRootNode() const
+{
+	TSharedRef<FFloatingPropertiesPropertyNode> LoopNode = SharedThis(const_cast<FFloatingPropertiesPropertyNode*>(this));
+	
+	while (TSharedPtr<FFloatingPropertiesPropertyNode> Parent = LoopNode->GetParent())
+	{
+		LoopNode = Parent.ToSharedRef();
+	}
+
+	return LoopNode;
+}
+
+TSharedRef<FFloatingPropertiesPropertyNode> FFloatingPropertiesPropertyNode::GetStackLeafMostNode() const
+{
+	TSharedRef<FFloatingPropertiesPropertyNode> LoopNode = SharedThis(const_cast<FFloatingPropertiesPropertyNode*>(this));
+
+	while (TSharedPtr<FFloatingPropertiesPropertyNode> Child = LoopNode->GetChild())
+	{
+		LoopNode = Child.ToSharedRef();
+	}
+
+	return LoopNode;
+}
+
+const FFloatingPropertiesClassPropertyAnchor& FFloatingPropertiesPropertyNode::GetPropertyAnchor() const
+{
+	return PropertyAnchor;
+}
+
+const FFloatingPropertiesClassPropertyPosition& FFloatingPropertiesPropertyNode::GetPropertyPosition() const
+{
+	return PropertyPosition;
+}
+
+FFloatingPropertiesClassPropertyPosition FFloatingPropertiesPropertyNode::CalculatePropertyPosition(const FVector2f& InDraggableArea, 
+	const FVector2f& InStackSize, const FVector2f& InPosition)
+{
+	FFloatingPropertiesClassPropertyPosition OutPosition = FFloatingPropertiesClassPropertyPosition::DefaultStackPosition;
+
+	constexpr float OneThird = 1.f / 3.f;
+	constexpr float TwoThirds = 2.f / 3.f;
+
+	const FVector2f TopLeft = InDraggableArea * OneThird;
+	const FVector2f BottomRight = InDraggableArea * TwoThirds;
+
+	if (InPosition.X < TopLeft.X)
+	{
+		OutPosition.HorizontalAnchor = EHorizontalAlignment::HAlign_Left;
+	}
+	else if (InPosition.X >= BottomRight.X)
+	{
+		OutPosition.HorizontalAnchor = EHorizontalAlignment::HAlign_Right;
+	}
+	else
+	{
+		OutPosition.HorizontalAnchor = EHorizontalAlignment::HAlign_Center;
+	}
+
+	if (InPosition.Y < TopLeft.Y)
+	{
+		OutPosition.VerticalAnchor = EVerticalAlignment::VAlign_Top;
+	}
+	else if (InPosition.Y >= BottomRight.Y)
+	{
+		OutPosition.VerticalAnchor = EVerticalAlignment::VAlign_Bottom;
+	}
+	else
+	{
+		OutPosition.VerticalAnchor = EVerticalAlignment::VAlign_Center;
+	}
+
+	const FVector2f AnchorMultiplier = CalculateAnchorMultiplier(OutPosition);
+	const FVector2f AvailableDragArea = InDraggableArea - InStackSize;
+	const FVector2f AnchorOffset = InPosition - (AnchorMultiplier * AvailableDragArea);
+
+	OutPosition.Offset.X = FMath::RoundToInt(AnchorOffset.X);
+	OutPosition.Offset.Y = FMath::RoundToInt(AnchorOffset.Y);
+
+	return OutPosition;
+}
+
+TSharedRef<SFloatingPropertiesPropertyWidget> FFloatingPropertiesPropertyNode::GetPropertyWidget() const
+{
+	return PropertyWidget;
+}
+
+EFloatingPropertiesUpdateResult FFloatingPropertiesPropertyNode::UpdateStackPositions(bool bInInvalidate)
+{
+	EFloatingPropertiesUpdateResult Result = EFloatingPropertiesUpdateResult::AlreadyUpToDate;
+	TSharedPtr<FFloatingPropertiesPropertyNode> LoopNode = SharedThis(this);
+
+	using namespace UE::FloatingProperties;
+
+	while (LoopNode.IsValid())
+	{
+		if (bInInvalidate)
+		{
+			LoopNode->bValidPosition = false;
+		}
+
+		Result = CombineUpdateResult(Result, LoopNode->UpdatePropertyNodePosition());
+
+		if (Result == EFloatingPropertiesUpdateResult::Failed)
+		{
+			return EFloatingPropertiesUpdateResult::Failed;
+		}
+
+		LoopNode = LoopNode->GetChild();
+	}
+
+	return Result;
+}
+
+void FFloatingPropertiesPropertyNode::InvalidateCachedPosition()
+{
+	bValidPosition = false;
+}
+
+const FVector2f& FFloatingPropertiesPropertyNode::GetCachedPosition() const
+{
+	return CachedPosition;
+}
+
+void FFloatingPropertiesPropertyNode::SetPropertyPositionDirect(const FFloatingPropertiesClassPropertyPosition& InPropertyPosition)
+{
+	PropertyPosition = InPropertyPosition;
+}
+
+void FFloatingPropertiesPropertyNode::SetCachedPosition(const FVector2f& InPosition)
+{
+	CachedPosition = InPosition;
+	bValidPosition = true;
+}
+
+void FFloatingPropertiesPropertyNode::SetPropertyPosition(const FVector2f& InDraggableArea, const FVector2f& InPosition)
+{
+	PropertyPosition = CalculatePropertyPosition(InDraggableArea, GetStackSize(), InPosition);
+
+	UpdateStackPositions(/* bInInvalidate */ true);
+}
+
+TArray<TSharedRef<FFloatingPropertiesPropertyNode>> FFloatingPropertiesPropertyNode::GetNodeStack() const
+{
+	TArray<TSharedRef<FFloatingPropertiesPropertyNode>> NodeStack;
+	NodeStack.Reserve(LinkedList->GetNodes().Num());
+
+	TSharedPtr<FFloatingPropertiesPropertyNode> LoopNode = GetStackRootNode();
+
+	while (LoopNode.IsValid())
+	{
+		NodeStack.Add(LoopNode.ToSharedRef());
+		LoopNode = LoopNode->GetChild();
+	}
+
+	return NodeStack;
+}
+
+FVector2f FFloatingPropertiesPropertyNode::GetStackSize() const
+{
+	FVector2f StackSize = FVector2f::ZeroVector;
+
+	for (const TSharedRef<FFloatingPropertiesPropertyNode>& Node : GetNodeStack())
+	{
+		const FVector2f DesiredSize = Node->GetPropertyWidget()->GetDesiredSize();
+
+		StackSize.X = FMath::Max(StackSize.X, DesiredSize.X);
+		StackSize.Y += DesiredSize.Y;
+	}
+
+	return StackSize;
+}
+
+EFloatingPropertiesUpdateResult FFloatingPropertiesPropertyNode::UpdatePropertyNodePosition()
+{
+	if (bValidPosition)
+	{
+		return EFloatingPropertiesUpdateResult::AlreadyUpToDate;
+	}
+
+	const FVector2f DraggableArea = LinkedList->GetDraggableArea();
+	constexpr float MinAxisValue = UE_SMALL_NUMBER;
+
+	if (DraggableArea.X < MinAxisValue || DraggableArea.Y < MinAxisValue)
+	{
+		return EFloatingPropertiesUpdateResult::Failed;
+	}
+
+	UFloatingPropertiesSettings* FloatingPropertiesSettings = GetMutableDefault<UFloatingPropertiesSettings>();
+	const FFloatingPropertiesClassProperty ClassProperty = PropertyWidget->GetClassProperty();
+
+	// Has a parent, so attach it to the parent's position
+	if (TSharedPtr<FFloatingPropertiesPropertyNode> ParentNode = GetParent())
+	{
+		if (ParentNode->MyIndex != MyIndex)
+		{
+			ParentNode->UpdatePropertyNodePosition();
+		}
+
+		const FVector2f& PreviousNodePosition = ParentNode->CachedPosition;
+		const FVector2f PreviousNodeSize = ParentNode->PropertyWidget->GetDesiredSize();
+
+		SetCachedPosition({
+			PreviousNodePosition.X,
+			PreviousNodePosition.Y + PreviousNodeSize.Y
+		});
+
+		if (GetStackRootNode()->GetPropertyPosition().OnDefaultStack())
+		{
+			FloatingPropertiesSettings->PropertyAnchors.Remove(ClassProperty);
+		}
+		else
+		{
+			FloatingPropertiesSettings->PropertyAnchors.FindOrAdd(ClassProperty) = PropertyAnchor;
+		}
+
+		FloatingPropertiesSettings->PropertyPositions.Remove(ClassProperty);
+	}
+	// No parent so use free position
+	else
+	{
+		const FVector2f StackSize = GetStackSize();
+
+		if (FMath::IsNearlyZero(StackSize.X) || FMath::IsNearlyZero(StackSize.Y)
+			|| StackSize.X < 0 || StackSize.Y < 0)
+		{
+			return EFloatingPropertiesUpdateResult::Failed;
+		}
+
+		const FVector2f AvailableDragArea = DraggableArea - StackSize;
+		const FVector2f AnchorMultiplier = CalculateAnchorMultiplier(PropertyPosition);
+
+		FVector2f Position = {
+			(AvailableDragArea.X * AnchorMultiplier.X) + static_cast<float>(PropertyPosition.Offset.X),
+			(AvailableDragArea.Y * AnchorMultiplier.Y) + static_cast<float>(PropertyPosition.Offset.Y)
+		};
+
+		Position.X = FMath::Clamp(Position.X, 0.f, AvailableDragArea.X);
+		Position.Y = FMath::Clamp(Position.Y, 0.f, AvailableDragArea.Y);
+
+		SetCachedPosition(Position);
+
+		if (PropertyPosition.OnDefaultStack())
+		{
+			FloatingPropertiesSettings->PropertyPositions.Remove(ClassProperty);
+		}
+		else
+		{
+			FloatingPropertiesSettings->PropertyPositions.FindOrAdd(ClassProperty) = PropertyPosition;
+		}
+		
+		FloatingPropertiesSettings->PropertyAnchors.Remove(ClassProperty);
+	}
+
+	return EFloatingPropertiesUpdateResult::Updated;
+}
+
+FVector2f FFloatingPropertiesPropertyNode::CalculateAnchorMultiplier(const FFloatingPropertiesClassPropertyPosition& InPropertyPosition)
+{
+	FVector2f AnchorMultiplier;
+
+	switch (InPropertyPosition.HorizontalAnchor.GetValue())
+	{
+		case EHorizontalAlignment::HAlign_Fill:
+		case EHorizontalAlignment::HAlign_Left:
+			AnchorMultiplier.X = 0.0f;
+			break;
+
+		case EHorizontalAlignment::HAlign_Center:
+			AnchorMultiplier.X = 0.5f;
+			break;
+
+		case EHorizontalAlignment::HAlign_Right:
+			AnchorMultiplier.X = 1.0f;
+			break;
+	}
+
+	switch (InPropertyPosition.VerticalAnchor.GetValue())
+	{
+		case EVerticalAlignment::VAlign_Fill:
+		case EVerticalAlignment::VAlign_Top:
+			AnchorMultiplier.Y = 0.0f;
+			break;
+
+		case EVerticalAlignment::VAlign_Center:
+			AnchorMultiplier.Y = 0.5f;
+			break;
+
+		case EVerticalAlignment::VAlign_Bottom:
+			AnchorMultiplier.Y = 1.0f;
+			break;
+	}
+
+	return AnchorMultiplier;
+}
diff --git a/Engine/Plugins/Experimental/FloatingProperties/Source/FloatingProperties/Private/Data/FloatingPropertiesPropertyNode.h b/Engine/Plugins/Experimental/FloatingProperties/Source/FloatingProperties/Private/Data/FloatingPropertiesPropertyNode.h
new file mode 100644
--- /dev/null
+++ b/Engine/Plugins/Experimental/FloatingProperties/Source/FloatingProperties/Private/Data/FloatingPropertiesPropertyNode.h
@@ -0,0 +1,119 @@
+// Copyright Epic Games, Inc. All Rights Reserved.
+
+#pragma once
+
+#include "FloatingPropertiesSettings.h"
+#include "Math/Vector2D.h"
+#include "Templates/SharedPointer.h"
+
+class FloatingPropertiesPropertyNodeContainer;
+class SFloatingPropertiesPropertyWidget;
+
+enum class EFloatingPropertiesUpdateResult : uint8
+{
+	Failed = 0,
+	AlreadyUpToDate,
+	Updated
+};
+
+namespace UE::FloatingProperties
+{
+	EFloatingPropertiesUpdateResult CombineUpdateResult(EFloatingPropertiesUpdateResult InCurrentResult, EFloatingPropertiesUpdateResult InNewResult);
+}
+
+class FFloatingPropertiesPropertyNode : public TSharedFromThis<FFloatingPropertiesPropertyNode>
+{
+public:
+	FFloatingPropertiesPropertyNode(FloatingPropertiesPropertyNodeContainer* InLinkedList,
+		TSharedRef<SFloatingPropertiesPropertyWidget> InPropertyWidget, int32 InIndex);
+
+	bool operator==(const FFloatingPropertiesPropertyNode& InOther) const
+	{
+		return MyIndex == InOther.MyIndex;
+	}
+
+	bool HasParent() const;
+
+	bool HasParent(TSharedRef<FFloatingPropertiesPropertyNode> InParent) const;
+
+	TSharedPtr<FFloatingPropertiesPropertyNode> GetParent() const;
+
+	bool HasChild() const;
+
+	bool HasChild(TSharedRef<FFloatingPropertiesPropertyNode> InChild) const;
+
+	TSharedPtr<FFloatingPropertiesPropertyNode> GetChild() const;
+
+	bool HasSibling(TSharedRef<FFloatingPropertiesPropertyNode> InSibling) const;
+
+	/* Add a parent relationship. */
+	void SetParent(TSharedPtr<FFloatingPropertiesPropertyNode> InParent);
+
+	/** Add child  relationship. */
+	void SetChild(TSharedRef<FFloatingPropertiesPropertyNode> InChild);
+
+	/** Insert between this parent and its child. */
+	void InsertAfter(TSharedRef<FFloatingPropertiesPropertyNode> InMiddle);
+
+	/** Remove parent and return it. */
+	TSharedPtr<FFloatingPropertiesPropertyNode> RemoveParent();
+
+	/** Remove child and return it. */
+	TSharedPtr<FFloatingPropertiesPropertyNode> RemoveChild();
+
+	void RemoveFromStack();
+
+	TSharedRef<FFloatingPropertiesPropertyNode> GetStackRootNode() const;
+
+	TSharedRef<FFloatingPropertiesPropertyNode> GetStackLeafMostNode() const;
+
+	const FFloatingPropertiesClassPropertyAnchor& GetPropertyAnchor() const;
+
+	const FFloatingPropertiesClassPropertyPosition& GetPropertyPosition() const;
+
+	TSharedRef<SFloatingPropertiesPropertyWidget> GetPropertyWidget() const;
+
+	/** Returns true if the node was able to update its position or if it already had a valid one. */
+	EFloatingPropertiesUpdateResult UpdatePropertyNodePosition();
+
+	/** Returns true if all nodes in the stack were able to update their positions. */
+	EFloatingPropertiesUpdateResult UpdateStackPositions(bool bInInvalidate);
+
+	void InvalidateCachedPosition();
+
+	const FVector2f& GetCachedPosition() const;
+
+	void SetPropertyPositionDirect(const FFloatingPropertiesClassPropertyPosition& InPropertyPosition);
+
+	void SetPropertyPosition(const FVector2f& InDraggableArea, const FVector2f& InPosition);
+
+	TArray<TSharedRef<FFloatingPropertiesPropertyNode>> GetNodeStack() const;
+
+	FVector2f GetStackSize() const;
+
+protected:
+	static FVector2f CalculateAnchorMultiplier(const FFloatingPropertiesClassPropertyPosition& InPropertyPosition);
+
+	static FFloatingPropertiesClassPropertyPosition CalculatePropertyPosition(const FVector2f& InDraggableArea,
+		const FVector2f& InStackSize, const FVector2f& InPosition);
+
+	FloatingPropertiesPropertyNodeContainer* LinkedList;
+
+	int32 MyIndex;
+
+	int32 ParentNodeIndex;
+
+	int32 ChildNodeIndex;
+
+	TSharedRef<SFloatingPropertiesPropertyWidget> PropertyWidget;
+
+	FFloatingPropertiesClassPropertyAnchor PropertyAnchor;
+
+	FFloatingPropertiesClassPropertyPosition PropertyPosition;
+
+	mutable FVector2f CachedPosition;
+
+	mutable bool bValidPosition;
+
+	void SetCachedPosition(const FVector2f& InPosition);
+};
diff --git a/Engine/Plugins/Experimental/FloatingProperties/Source/FloatingProperties/Private/Data/FloatingPropertiesPropertyNodeContainer.cpp b/Engine/Plugins/Experimental/FloatingProperties/Source/FloatingProperties/Private/Data/FloatingPropertiesPropertyNodeContainer.cpp
new file mode 100644
--- /dev/null
+++ b/Engine/Plugins/Experimental/FloatingProperties/Source/FloatingProperties/Private/Data/FloatingPropertiesPropertyNodeContainer.cpp
@@ -0,0 +1,81 @@
+// Copyright Epic Games, Inc. All Rights Reserved.
+
+#include "Data/FloatingPropertiesPropertyNodeContainer.h"
+#include "Data/FloatingPropertiesPropertyNode.h"
+#include "Math/Vector2D.h"
+#include "Widgets/SFloatingPropertiesPropertyWidget.h"
+#include "Widgets/SFloatingPropertiesViewportWidget.h"
+
+FloatingPropertiesPropertyNodeContainer::FloatingPropertiesPropertyNodeContainer(TSharedRef<SFloatingPropertiesViewportWidget> InViewportWidget)
+{
+	ContainerWeak = InViewportWidget;
+}
+
+FVector2f FloatingPropertiesPropertyNodeContainer::GetDraggableArea() const
+{
+	if (TSharedPtr<SFloatingPropertiesViewportWidget> ContainerWidget = ContainerWeak.Pin())
+	{
+		return ContainerWidget->GetDraggableArea();
+	}
+
+	return FVector2f::ZeroVector;
+}
+
+TSharedPtr<FFloatingPropertiesPropertyNode> FloatingPropertiesPropertyNodeContainer::FindNodeForWidget(
+	TSharedRef<SFloatingPropertiesPropertyWidget> InWidget) const
+{
+	const int32* NodeIndexPtr = WidgetToNode.Find(InWidget);
+
+	if (NodeIndexPtr)
+	{
+		return PropertyNodes[*NodeIndexPtr];
+	}
+
+	return nullptr;
+}
+
+TSharedRef<FFloatingPropertiesPropertyNode> FloatingPropertiesPropertyNodeContainer::AddWidget(TSharedRef<SFloatingPropertiesPropertyWidget> InWidget)
+{
+	const int32 PropertyIndex = PropertyNodes.Num();
+
+	TSharedRef<FFloatingPropertiesPropertyNode> Node = MakeShared<FFloatingPropertiesPropertyNode>(this, InWidget, PropertyIndex);
+
+	PropertyNodes.Add(Node);
+	WidgetToNode.Add(InWidget, PropertyIndex);
+
+	return Node;
+}
+
+void FloatingPropertiesPropertyNodeContainer::Reserve(int32 InAmount)
+{
+	WidgetToNode.Reserve(InAmount);
+	PropertyNodes.Reserve(InAmount);
+}
+
+void FloatingPropertiesPropertyNodeContainer::Empty()
+{
+	WidgetToNode.Empty();
+	PropertyNodes.Empty();
+}
+
+void FloatingPropertiesPropertyNodeContainer::InvalidateAllPositions() const
+{
+	for (const TSharedRef<FFloatingPropertiesPropertyNode>& Node : PropertyNodes)
+	{
+		Node->InvalidateCachedPosition();
+	}
+}
+
+EFloatingPropertiesUpdateResult FloatingPropertiesPropertyNodeContainer::EnsureCachedPositions() const
+{
+	using namespace UE::FloatingProperties;
+
+	EFloatingPropertiesUpdateResult Result = EFloatingPropertiesUpdateResult::AlreadyUpToDate;
+
+	for (const TSharedRef<FFloatingPropertiesPropertyNode>& Node : PropertyNodes)
+	{
+		Result = CombineUpdateResult(Result, Node->UpdatePropertyNodePosition());
+	}
+
+	return Result;
+}
diff --git a/Engine/Plugins/Experimental/FloatingProperties/Source/FloatingProperties/Private/Data/FloatingPropertiesPropertyNodeContainer.h b/Engine/Plugins/Experimental/FloatingProperties/Source/FloatingProperties/Private/Data/FloatingPropertiesPropertyNodeContainer.h
new file mode 100644
--- /dev/null
+++ b/Engine/Plugins/Experimental/FloatingProperties/Source/FloatingProperties/Private/Data/FloatingPropertiesPropertyNodeContainer.h
@@ -0,0 +1,44 @@
+// Copyright Epic Games, Inc. All Rights Reserved.
+
+#pragma once
+
+#include "FloatingPropertiesSettings.h"
+#include "Containers/Array.h"
+#include "Containers/ArrayView.h"
+#include "Containers/Map.h"
+#include "Math/MathFwd.h"
+#include "Templates/SharedPointer.h"
+
+class FFloatingPropertiesPropertyNode;
+class SFloatingPropertiesPropertyWidget;
+class SFloatingPropertiesViewportWidget;
+enum class EFloatingPropertiesUpdateResult : uint8;
+
+class FloatingPropertiesPropertyNodeContainer
+{
+public:
+	FloatingPropertiesPropertyNodeContainer(TSharedRef<SFloatingPropertiesViewportWidget> InViewportWidget);
+
+	TConstArrayView<TSharedRef<FFloatingPropertiesPropertyNode>> GetNodes() const { return PropertyNodes; }
+
+	void Reserve(int32 InAmount);
+
+	TSharedRef<FFloatingPropertiesPropertyNode> AddWidget(TSharedRef<SFloatingPropertiesPropertyWidget> InWidget);
+
+	TSharedPtr<FFloatingPropertiesPropertyNode> FindNodeForWidget(TSharedRef<SFloatingPropertiesPropertyWidget> InWidget) const;
+
+	void Empty();
+
+	void InvalidateAllPositions() const;
+
+	EFloatingPropertiesUpdateResult EnsureCachedPositions() const;
+
+	FVector2f GetDraggableArea() const;
+
+protected:
+	TWeakPtr<SFloatingPropertiesViewportWidget> ContainerWeak;
+
+	TArray<TSharedRef<FFloatingPropertiesPropertyNode>> PropertyNodes;
+
+	TMap<TWeakPtr<SFloatingPropertiesPropertyWidget>, int32> WidgetToNode;
+};
\ No newline at end of file
diff --git a/Engine/Plugins/Experimental/FloatingProperties/Source/FloatingProperties/Private/Data/FloatingPropertiesSnapMetrics.cpp b/Engine/Plugins/Experimental/FloatingProperties/Source/FloatingProperties/Private/Data/FloatingPropertiesSnapMetrics.cpp
new file mode 100644
--- /dev/null
+++ b/Engine/Plugins/Experimental/FloatingProperties/Source/FloatingProperties/Private/Data/FloatingPropertiesSnapMetrics.cpp
@@ -0,0 +1,53 @@
+// Copyright Epic Games, Inc. All Rights Reserved.
+
+#include "Data/FloatingPropertiesSnapMetrics.h"
+#include "Data/FloatingPropertiesPropertyNode.h"
+#include "Widgets/SFloatingPropertiesPropertyWidget.h"
+
+FFloatingPropertiesSnapMetrics FFloatingPropertiesSnapMetrics::Make(TSharedRef<FFloatingPropertiesPropertyNode> InNode,
+	const FVector2f& InSnapPosition, EFloatingPropertiesSnapType InAttachType)
+{
+	FFloatingPropertiesSnapMetrics Metrics = {
+		InNode,
+		GetSnapPosition(InNode, InAttachType),
+		0.f,
+		InAttachType
+	};
+
+	Metrics.DistanceSq = (Metrics.Position - InSnapPosition).SizeSquared();
+
+	return Metrics;
+}
+
+FVector2f FFloatingPropertiesSnapMetrics::GetSnapPosition(TSharedRef<FFloatingPropertiesPropertyNode> InNode, EFloatingPropertiesSnapType InAttachType)
+{
+	FVector2f Position = InNode->GetPropertyWidget()->GetTickSpaceGeometry().GetAbsolutePosition();
+
+	if (InAttachType == EFloatingPropertiesSnapType::AttachAsParent)
+	{
+		Position.Y += InNode->GetPropertyWidget()->GetDesiredSize().Y;
+	}
+
+	return Position;
+}
+
+void FFloatingPropertiesSnapMetrics::SnapToDraggableArea(const FVector2f& InDraggableArea, FVector2f& InOutPosition)
+{
+	if (InOutPosition.X <= FFloatingPropertiesSnapMetrics::SnapDistance)
+	{
+		InOutPosition.X = 0.f;
+	}
+	else if (InOutPosition.X >= (InDraggableArea.X - FFloatingPropertiesSnapMetrics::SnapDistance))
+	{
+		InOutPosition.X = InDraggableArea.X;
+	}
+
+	if (InOutPosition.Y <= FFloatingPropertiesSnapMetrics::SnapDistance)
+	{
+		InOutPosition.Y = 0.f;
+	}
+	else if (InOutPosition.Y >= (InDraggableArea.Y - FFloatingPropertiesSnapMetrics::SnapDistance))
+	{
+		InOutPosition.Y = InDraggableArea.Y;
+	}
+}
diff --git a/Engine/Plugins/Experimental/FloatingProperties/Source/FloatingProperties/Private/Data/FloatingPropertiesSnapMetrics.h b/Engine/Plugins/Experimental/FloatingProperties/Source/FloatingProperties/Private/Data/FloatingPropertiesSnapMetrics.h
new file mode 100644
--- /dev/null
+++ b/Engine/Plugins/Experimental/FloatingProperties/Source/FloatingProperties/Private/Data/FloatingPropertiesSnapMetrics.h
@@ -0,0 +1,35 @@
+// Copyright Epic Games, Inc. All Rights Reserved.
+
+#pragma once
+
+#include "Math/Vector2D.h"
+#include "Templates/SharedPointer.h"
+
+class FFloatingPropertiesPropertyNode;
+
+enum class EFloatingPropertiesSnapType
+{
+	AttachAsParent,
+	AttachAsChild
+};
+
+struct FFloatingPropertiesSnapMetrics
+{
+	static constexpr int32 TopLeft = 0;
+	static constexpr int32 BottomLeft = 1;
+	static constexpr float SnapDistance = 10.f;
+	static constexpr float SnapDistanceSq = SnapDistance * SnapDistance;
+	static constexpr float SnapBreakDistanceSq = SnapDistanceSq * 2.f;
+
+	static FFloatingPropertiesSnapMetrics Make(TSharedRef<FFloatingPropertiesPropertyNode> InNode,
+		const FVector2f& InSnapPosition, EFloatingPropertiesSnapType InAttachType);
+
+	static FVector2f GetSnapPosition(TSharedRef<FFloatingPropertiesPropertyNode> InNode, EFloatingPropertiesSnapType InAttachType);
+
+	static void SnapToDraggableArea(const FVector2f& InDraggableArea, FVector2f& InOutPosition);
+
+	TSharedRef<FFloatingPropertiesPropertyNode> Node;
+	FVector2f Position;
+	float DistanceSq;
+	EFloatingPropertiesSnapType AttachType;
+};
diff --git a/Engine/Plugins/Experimental/FloatingProperties/Source/FloatingProperties/Private/Data/FloatingPropertiesWidgetController.cpp b/Engine/Plugins/Experimental/FloatingProperties/Source/FloatingProperties/Private/Data/FloatingPropertiesWidgetController.cpp
new file mode 100644
--- /dev/null
+++ b/Engine/Plugins/Experimental/FloatingProperties/Source/FloatingProperties/Private/Data/FloatingPropertiesWidgetController.cpp
@@ -0,0 +1,329 @@
+// Copyright Epic Games, Inc. All Rights Reserved.
+
+#include "Data/FloatingPropertiesWidgetController.h"
+#include "Algo/Reverse.h"
+#include "Components/SceneComponent.h"
+#include "Containers/Array.h"
+#include "Data/IFloatingPropertiesWidgetContainer.h"
+#include "Elements/Framework/TypedElementSelectionSet.h"
+#include "FloatingPropertiesModule.h"
+#include "IDetailTreeNode.h"
+#include "IPropertyRowGenerator.h"
+#include "Modules/ModuleManager.h"
+#include "PropertyEditorModule.h"
+#include "PropertyHandle.h"
+#include "Selection.h"
+#include "UObject/Class.h"
+#include "UObject/TextProperty.h"
+#include "UObject/UnrealType.h"
+#include "Widgets/SFloatingPropertiesViewportWidget.h"
+
+namespace UE::FloatingFavorite::Private
+{
+	const TSet<FFieldClass*> AllowedClasses = {
+		FBoolProperty::StaticClass(),
+		FEnumProperty::StaticClass(),
+		FNameProperty::StaticClass(),
+		FNumericProperty::StaticClass(),
+		FObjectPropertyBase::StaticClass(),
+		FStrProperty::StaticClass(),
+		FTextProperty::StaticClass(),
+		FStructProperty::StaticClass()
+	};
+
+	bool IsValidPropertyType(const FProperty& InProperty)
+	{
+		FFieldClass* PropertyClass = InProperty.GetClass();
+
+		while (true)
+		{
+			if (AllowedClasses.Contains(PropertyClass))
+			{
+				return true;
+			}
+
+			PropertyClass = PropertyClass->GetSuperClass();
+
+			if (!PropertyClass || PropertyClass == FProperty::StaticClass())
+			{
+				return false;
+			}
+		}
+	}
+
+	bool IsFullPropertyStruct(const FStructProperty& InProperty)
+	{
+		const FFloatingPropertiesModule::FCreateStructPropertyValueWidgetDelegate* Delegate
+			= FFloatingPropertiesModule::Get().GetStructPropertyValueWidgetDelegate(InProperty.Struct);
+
+		return (Delegate && Delegate->IsBound());
+	}
+}
+
+FFloatingPropertiesWidgetController::FFloatingPropertiesWidgetController(TSharedRef<IFloatingPropertiesDataProvider> InDataProvider)
+	: DataProvider(InDataProvider)
+{
+}
+
+FFloatingPropertiesWidgetController::~FFloatingPropertiesWidgetController()
+{
+	UnregisterSelectionChanged();
+	RemoveWidgets();
+}
+
+void FFloatingPropertiesWidgetController::Init()
+{
+	RegisterSelectionChanged();
+	OnViewportChange();
+}
+
+void FFloatingPropertiesWidgetController::OnViewportChange()
+{
+	const TArray<TSharedRef<IFloatingPropertiesWidgetContainer>>& Containers = DataProvider->GetWidgetContainers();
+
+	bool bCreatedWidget = false;
+
+	for (const TSharedRef<IFloatingPropertiesWidgetContainer>& Container : Containers)
+	{
+		if (Container->GetWidget().IsValid())
+		{
+			continue;
+		}
+
+		Container->SetWidget(CreateWidget(FIsVisibleDelegate::CreateSPLambda(this,
+			[this, Container]()
+			{
+				return DataProvider->IsWidgetVisibleInContainer(Container);
+			})));
+
+		bCreatedWidget = true;
+	}
+
+	if (bCreatedWidget)
+	{
+		SelectedActorWeak.Reset();
+		SelectedComponentWeak.Reset();
+		OnSelectionChange(nullptr);
+	}
+}
+
+void FFloatingPropertiesWidgetController::RegisterSelectionChanged()
+{
+	if (USelection* ActorSelection = DataProvider->GetActorSelection())
+	{
+		if (UTypedElementSelectionSet* SelectionSet = ActorSelection->GetElementSelectionSet())
+		{
+			SelectionSet->OnChanged().AddSP(this, &FFloatingPropertiesWidgetController::OnSelectionChange);
+		}
+	}
+
+	if (USelection* ComponentSelection = DataProvider->GetComponentSelection())
+	{
+		if (UTypedElementSelectionSet* SelectionSet = ComponentSelection->GetElementSelectionSet())
+		{
+			SelectionSet->OnChanged().AddSP(this, &FFloatingPropertiesWidgetController::OnSelectionChange);
+		}
+	}
+}
+
+void FFloatingPropertiesWidgetController::UnregisterSelectionChanged()
+{
+	if (DataProvider.IsValid())
+	{
+		if (USelection* ActorSelection = DataProvider->GetActorSelection())
+		{
+			if (UTypedElementSelectionSet* SelectionSet = ActorSelection->GetElementSelectionSet())
+			{
+				SelectionSet->OnChanged().RemoveAll(this);
+			}
+		}
+
+		if (USelection* ComponentSelection = DataProvider->GetComponentSelection())
+		{
+			if (UTypedElementSelectionSet* SelectionSet = ComponentSelection->GetElementSelectionSet())
+			{
+				SelectionSet->OnChanged().RemoveAll(this);
+			}
+		}
+	}
+}
+
+void FFloatingPropertiesWidgetController::RemoveWidgets()
+{
+	for (const TSharedRef<IFloatingPropertiesWidgetContainer>& Container : DataProvider->GetWidgetContainers())
+	{
+		Container->RemoveWidget();
+	}
+}
+
+void FFloatingPropertiesWidgetController::OnSelectionChange(const UTypedElementSelectionSet* InSelectionSet)
+{
+	OnSelectionChange();
+}
+
+void FFloatingPropertiesWidgetController::OnSelectionChange()
+{
+	OnSelectionChange(
+		DataProvider->GetWorld(),
+		DataProvider->GetActorSelection(),
+		DataProvider->GetComponentSelection()
+	);
+}
+
+void FFloatingPropertiesWidgetController::OnSelectionChange(UWorld* InWorld, USelection* InActorSelection, 
+	USelection* InComponentSelection)
+{
+	if (!InWorld || !InActorSelection || !InComponentSelection)
+	{
+		return;
+	}
+
+	AActor* SelectedActor = nullptr;
+	UActorComponent* SelectedComponent = nullptr;
+
+	TArray<UActorComponent*> SelectedComponents;
+	InComponentSelection->GetSelectedObjects(SelectedComponents);
+
+	for (UActorComponent* Component : SelectedComponents)
+	{
+		if (IsValid(Component) && Component->GetWorld() == InWorld)
+		{
+			if (AActor* Owner = Component->GetOwner())
+			{
+				if (Owner->GetRootComponent() != Component)
+				{
+					SelectedComponent = Component;
+					break;
+				}
+			}
+		}
+	}
+
+	if (SelectedComponent == nullptr)
+	{
+		TArray<AActor*> SelectedActors;
+		InActorSelection->GetSelectedObjects(SelectedActors);
+
+		for (AActor* Actor : SelectedActors)
+		{
+			if (IsValid(Actor) && Actor->GetWorld() == InWorld)
+			{
+				SelectedActor = Actor;
+				break;
+			}
+		}
+	}
+
+	if (SelectedComponentWeak.Get() == SelectedComponent && SelectedActorWeak.Get() == SelectedActor)
+	{
+ 		return;
+	}
+
+	SelectedActorWeak = SelectedActor;
+	ActorProperties.Properties.Empty();
+	ActorProperties.Generator.Reset();
+	UpdatePropertiesForObject(SelectedActor, ActorProperties);
+
+	SelectedComponentWeak = SelectedComponent;
+	ComponentProperties.Properties.Empty();
+	ComponentProperties.Generator.Reset();
+	UpdatePropertiesForObject(SelectedComponent, ComponentProperties);
+
+	RebuildWidgets(ActorProperties, ComponentProperties);
+}
+
+void FFloatingPropertiesWidgetController::UpdatePropertiesForObject(UObject* InObject, FFloatingPropertiesPropertyList& OutPropertyList)
+{
+	OutPropertyList.Generator.Reset();
+	OutPropertyList.Properties.Empty();
+
+	if (!InObject)
+	{
+		return;
+	}
+
+	const bool bObjectIsActor = InObject->IsA<AActor>();
+
+	FPropertyEditorModule& PropertyEditorModule = FModuleManager::GetModuleChecked<FPropertyEditorModule>("PropertyEditor");
+	
+	FPropertyRowGeneratorArgs Args;
+	Args.DefaultsOnlyVisibility = EEditDefaultsOnlyNodeVisibility::Hide;
+	Args.bAllowEditingClassDefaultObjects = false;
+	Args.bAllowMultipleTopLevelObjects = false;
+	Args.bShouldShowHiddenProperties = false;
+
+	TSharedRef<IPropertyRowGenerator> PropertyRowGenerator = PropertyEditorModule.CreatePropertyRowGenerator(Args);
+	PropertyRowGenerator->SetObjects({InObject});
+
+	const TArray<TSharedRef<IDetailTreeNode>>& RootNodes = PropertyRowGenerator->GetRootTreeNodes();
+	TArray<TSharedRef<IDetailTreeNode>> Stack(RootNodes);
+
+	Algo::Reverse(Stack);
+
+	Stack.Reserve(100); // There's going to be a lot.
+
+	while (!Stack.IsEmpty())
+	{
+		TSharedRef<IDetailTreeNode> Node = Stack.Pop(EAllowShrinking::No);
+
+		if (Node->GetNodeType() == EDetailNodeType::Item)
+		{
+			if (TSharedPtr<IPropertyHandle> PropertyHandle = Node->CreatePropertyHandle())
+			{
+				FProperty* Property = PropertyHandle->GetProperty();
+
+				if (!Property)
+				{
+					continue;
+				}
+
+				using namespace UE::FloatingFavorite::Private;
+
+				if (!IsValidPropertyType(*Property))
+				{
+					continue;
+				}
+
+				if (PropertyHandle->IsFavorite())
+				{
+					OutPropertyList.Properties.Add({PropertyHandle, Property, Node});
+				}
+			}
+		}
+
+		TArray<TSharedRef<IDetailTreeNode>> Children;
+		Node->GetChildren(Children, true);
+
+		Algo::Reverse(Children);
+
+		Stack.Append(Children);
+	}
+
+	if (!OutPropertyList.Properties.IsEmpty())
+	{
+		OutPropertyList.Generator = PropertyRowGenerator;
+	}
+}
+
+TSharedRef<SFloatingPropertiesViewportWidget> FFloatingPropertiesWidgetController::CreateWidget(FIsVisibleDelegate InVisibilityCallback)
+{
+	return SNew(SFloatingPropertiesViewportWidget)
+		.IsVisible(InVisibilityCallback);
+}
+
+void FFloatingPropertiesWidgetController::RebuildWidgets(const FFloatingPropertiesPropertyList& InActorProperties,
+	const FFloatingPropertiesPropertyList& InComponentProperties)
+{
+	for (const TSharedRef<IFloatingPropertiesWidgetContainer>& Container : DataProvider->GetWidgetContainers())
+	{
+		if (TSharedPtr<SFloatingPropertiesViewportWidget> Widget = Container->GetWidget())
+		{
+			Widget->BuildWidget(
+				SelectedActorWeak.Get(),
+				InActorProperties,
+				SelectedComponentWeak.Get(),
+				InComponentProperties
+			);
+		}
+	}
+}
diff --git a/Engine/Plugins/Experimental/FloatingProperties/Source/FloatingProperties/Private/FloatingPropertiesCommands.cpp b/Engine/Plugins/Experimental/FloatingProperties/Source/FloatingProperties/Private/FloatingPropertiesCommands.cpp
new file mode 100644
--- /dev/null
+++ b/Engine/Plugins/Experimental/FloatingProperties/Source/FloatingProperties/Private/FloatingPropertiesCommands.cpp
@@ -0,0 +1,26 @@
+// Copyright Epic Games, Inc. All Rights Reserved.
+
+#include "FloatingPropertiesCommands.h"
+#include "Styling/AppStyle.h"
+
+#define LOCTEXT_NAMESPACE "FloatingPropertiesCommands"
+
+FFloatingPropertiesCommands::FFloatingPropertiesCommands()
+	: TCommands<FFloatingPropertiesCommands>(TEXT("FloatingProperties")
+		, LOCTEXT("FloatingProperties", "FloatingProperties")
+		, NAME_None
+		, FAppStyle::GetAppStyleSetName()
+	)
+{
+}
+
+void FFloatingPropertiesCommands::RegisterCommands()
+{
+	UI_COMMAND(ToggleEnabled
+		, "Floating Properties"
+		, "Toggle floating properties on and off."
+		, EUserInterfaceActionType::ToggleButton
+		, FInputChord());	
+}
+
+#undef LOCTEXT_NAMESPACE
diff --git a/Engine/Plugins/Experimental/FloatingProperties/Source/FloatingProperties/Private/FloatingPropertiesCommands.h b/Engine/Plugins/Experimental/FloatingProperties/Source/FloatingProperties/Private/FloatingPropertiesCommands.h
new file mode 100644
--- /dev/null
+++ b/Engine/Plugins/Experimental/FloatingProperties/Source/FloatingProperties/Private/FloatingPropertiesCommands.h
@@ -0,0 +1,17 @@
+// Copyright Epic Games, Inc. All Rights Reserved.
+
+#pragma once
+
+#include "Framework/Commands/Commands.h"
+
+class FFloatingPropertiesCommands : public TCommands<FFloatingPropertiesCommands>
+{
+public:
+	FFloatingPropertiesCommands();
+
+	//~ Begin TCommands
+	virtual void RegisterCommands() override;
+	//~ End TCommands
+
+	TSharedPtr<FUICommandInfo> ToggleEnabled;
+};
diff --git a/Engine/Plugins/Experimental/FloatingProperties/Source/FloatingProperties/Private/FloatingPropertiesModule.cpp b/Engine/Plugins/Experimental/FloatingProperties/Source/FloatingProperties/Private/FloatingPropertiesModule.cpp
new file mode 100644
--- /dev/null
+++ b/Engine/Plugins/Experimental/FloatingProperties/Source/FloatingProperties/Private/FloatingPropertiesModule.cpp
@@ -0,0 +1,165 @@
+// Copyright Epic Games, Inc. All Rights Reserved.
+
+#include "FloatingPropertiesModule.h"
+#include "FloatingPropertiesCommands.h"
+#include "FloatingPropertiesSettings.h"
+#include "LevelEditor.h"
+#include "LevelEditor/FloatingPropertiesLevelEditorWidgetController.h"
+#include "Modules/ModuleManager.h"
+#include "Templates/SharedPointer.h"
+#include "Widgets/CustomPropertyWidgets/SFloatingPropertiesColorPropertyEditor.h"
+#include "Widgets/CustomPropertyWidgets/SFloatingPropertiesLinearColorPropertyEditor.h"
+
+DEFINE_LOG_CATEGORY(LogFloatingProperties);
+
+FFloatingPropertiesModule& FFloatingPropertiesModule::Get()
+{
+	return FModuleManager::GetModuleChecked<FFloatingPropertiesModule>("FloatingProperties");
+}
+
+void FFloatingPropertiesModule::StartupModule()
+{
+	FFloatingPropertiesCommands::Register();
+
+	UFloatingPropertiesSettings::OnChange.AddRaw(this, &FFloatingPropertiesModule::OnSettingsChanged);
+
+	FLevelEditorModule& LevelEditorModule = FModuleManager::LoadModuleChecked<FLevelEditorModule>("LevelEditor");
+	LevelEditorModule.OnLevelEditorCreated().AddRaw(this, &FFloatingPropertiesModule::OnLevelEditorCreated);
+
+	AddDefaultStructPropertyValueWidgetDelegates();
+}
+
+void FFloatingPropertiesModule::ShutdownModule()
+{
+	FFloatingPropertiesCommands::Unregister();
+
+	if (FLevelEditorModule* LevelEditorModule = FModuleManager::GetModulePtr<FLevelEditorModule>("LevelEditor"))
+	{
+		LevelEditorModule->OnLevelEditorCreated().RemoveAll(this);
+	}
+
+	DestroyLevelEditorWidgetController();
+
+	UFloatingPropertiesSettings::OnChange.RemoveAll(this);
+}
+
+void FFloatingPropertiesModule::RegiserStructPropertyValueWidgetDelegate(UScriptStruct* InStruct,
+	FCreateStructPropertyValueWidgetDelegate InDelegate)
+{
+	if (!InStruct)
+	{
+		return;
+	}
+
+	if (!InDelegate.IsBound())
+	{
+		return;
+	}
+
+	const FName StructName = InStruct->GetFName();
+
+	PropertyValueWidgetDelegates.Add(StructName, InDelegate);
+}
+
+const FFloatingPropertiesModule::FCreateStructPropertyValueWidgetDelegate* FFloatingPropertiesModule::GetStructPropertyValueWidgetDelegate(
+	UScriptStruct* InStruct) const
+{
+	if (!InStruct)
+	{
+		return nullptr;
+	}
+
+	const FName StructName = InStruct->GetFName();
+
+	return PropertyValueWidgetDelegates.Find(StructName);
+}
+
+void FFloatingPropertiesModule::UnregiserStructPropertyValueWidgetDelegate(UScriptStruct* InStruct)
+{
+	if (!InStruct)
+	{
+		return;
+	}
+
+	const FName StructName = InStruct->GetFName();
+
+	PropertyValueWidgetDelegates.Remove(StructName);
+}
+
+void FFloatingPropertiesModule::OnLevelEditorCreated(TSharedPtr<ILevelEditor> InLevelEditor)
+{
+	CreateLevelEditorWidgetController(InLevelEditor);
+}
+
+void FFloatingPropertiesModule::CreateLevelEditorWidgetController(TSharedPtr<ILevelEditor> InLevelEditor)
+{
+	if (LevelEditorWidgetController.IsValid())
+	{
+		return;
+	}
+
+	if (!InLevelEditor.IsValid())
+	{
+		return;
+	}
+
+	FFloatingPropertiesLevelEditorWidgetController::StaticInit();
+
+	if (const UFloatingPropertiesSettings* FloatingPropertiesSettings = GetDefault<UFloatingPropertiesSettings>())
+	{
+		if (FloatingPropertiesSettings->bEnabled)
+		{
+			LevelEditorWidgetController = MakeShared<FFloatingPropertiesLevelEditorWidgetController>(InLevelEditor.ToSharedRef());
+			LevelEditorWidgetController->Init();
+		}
+	}	
+}
+
+void FFloatingPropertiesModule::DestroyLevelEditorWidgetController()
+{
+	LevelEditorWidgetController.Reset();
+}
+
+void FFloatingPropertiesModule::OnSettingsChanged(const UFloatingPropertiesSettings* InSettings, FName InSetting)
+{
+	static const FName EnabledName = GET_MEMBER_NAME_CHECKED(UFloatingPropertiesSettings, bEnabled);
+
+	if (InSetting == EnabledName)
+	{
+		if (InSettings->bEnabled)
+		{
+			if (!LevelEditorWidgetController.IsValid())
+			{
+				FLevelEditorModule& LevelEditorModule = FModuleManager::LoadModuleChecked<FLevelEditorModule>("LevelEditor");
+				TSharedPtr<ILevelEditor> LevelEditor = LevelEditorModule.GetFirstLevelEditor();
+
+				if (LevelEditor.IsValid())
+				{
+					CreateLevelEditorWidgetController(LevelEditor.ToSharedRef());
+				}
+			}
+		}
+		else
+		{
+			if (LevelEditorWidgetController.IsValid())
+			{
+				DestroyLevelEditorWidgetController();
+			}
+		}
+	}
+}
+
+void FFloatingPropertiesModule::AddDefaultStructPropertyValueWidgetDelegates()
+{
+	RegiserStructPropertyValueWidgetDelegate(
+		TBaseStructure<FColor>::Get(),
+		FCreateStructPropertyValueWidgetDelegate::CreateStatic(&SFloatingPropertiesColorPropertyEditor::CreateWidget)
+	);
+
+	RegiserStructPropertyValueWidgetDelegate(
+		TBaseStructure<FLinearColor>::Get(),
+		FCreateStructPropertyValueWidgetDelegate::CreateStatic(&SFloatingPropertiesLinearColorPropertyEditor::CreateWidget)
+	);
+}
+
+IMPLEMENT_MODULE(FFloatingPropertiesModule, FloatingProperties)
diff --git a/Engine/Plugins/Experimental/FloatingProperties/Source/FloatingProperties/Private/FloatingPropertiesSettings.cpp b/Engine/Plugins/Experimental/FloatingProperties/Source/FloatingProperties/Private/FloatingPropertiesSettings.cpp
new file mode 100644
--- /dev/null
+++ b/Engine/Plugins/Experimental/FloatingProperties/Source/FloatingProperties/Private/FloatingPropertiesSettings.cpp
@@ -0,0 +1,59 @@
+// Copyright Epic Games, Inc. All Rights Reserved.
+
+#include "FloatingPropertiesSettings.h"
+#include "Containers/UnrealString.h"
+
+UFloatingPropertiesSettings::FOnFloatingPropertiesSettingsChanged UFloatingPropertiesSettings::OnChange;
+
+const FFloatingPropertiesClassPropertyPosition FFloatingPropertiesClassPropertyPosition::DefaultStackPosition = {
+		EHorizontalAlignment::HAlign_Left,
+		EVerticalAlignment::VAlign_Bottom,
+		FIntPoint::ZeroValue
+};
+
+bool FFloatingPropertiesClassPropertyPosition::OnDefaultStack() const
+{
+	return HorizontalAnchor == DefaultStackPosition.HorizontalAnchor
+		&& VerticalAnchor == DefaultStackPosition.VerticalAnchor
+		&& Offset == DefaultStackPosition.Offset;
+}
+
+FString UFloatingPropertiesSettings::FindUniqueName(const FFloatingPropertiesClassProperties& InPropertyList, const FString& InDefault)
+{
+	if (!InPropertyList.Properties.Contains(InDefault))
+	{
+		return InDefault;
+	}
+
+	for (int32 Index = 1; true; ++Index)
+	{
+		FString UniqueTest = InDefault + "_" + FString::FromInt(Index);
+
+		if (!InPropertyList.Properties.Contains(UniqueTest))
+		{
+			return UniqueTest;
+		}
+	}
+
+	return InDefault;
+}
+
+UFloatingPropertiesSettings::UFloatingPropertiesSettings()
+{
+	CategoryName = TEXT("Plugins");
+	SectionName = TEXT("Floating Properties");
+
+	bEnabled = true;
+}
+
+void UFloatingPropertiesSettings::PostEditChangeProperty(FPropertyChangedEvent& PropertyChangedEvent)
+{
+	Super::PostEditChangeProperty(PropertyChangedEvent);
+
+	const FName MemberName = PropertyChangedEvent.GetMemberPropertyName();
+
+	if (MemberName != NAME_None)
+	{
+		OnChange.Broadcast(this, MemberName);
+	}
+}
diff --git a/Engine/Plugins/Experimental/FloatingProperties/Source/FloatingProperties/Private/LevelEditor/FloatingPropertiesLevelEditorDataProvider.cpp b/Engine/Plugins/Experimental/FloatingProperties/Source/FloatingProperties/Private/LevelEditor/FloatingPropertiesLevelEditorDataProvider.cpp
new file mode 100644
--- /dev/null
+++ b/Engine/Plugins/Experimental/FloatingProperties/Source/FloatingProperties/Private/LevelEditor/FloatingPropertiesLevelEditorDataProvider.cpp
@@ -0,0 +1,75 @@
+// Copyright Epic Games, Inc. All Rights Reserved.
+
+#include "FloatingPropertiesLevelEditorDataProvider.h"
+#include "LevelEditor.h"
+#include "LevelEditor/FloatingPropertiesLevelEditorWidgetContainer.h"
+#include "Selection.h"
+#include "SLevelViewport.h"
+
+FFloatingPropertiesLevelEditorDataProvider::FFloatingPropertiesLevelEditorDataProvider(TSharedRef<ILevelEditor> InLevelEditor)
+	: FFloatingPropertiesToolkitHostDataProvider(InLevelEditor)
+{
+}
+
+TArray<TSharedRef<IFloatingPropertiesWidgetContainer>> FFloatingPropertiesLevelEditorDataProvider::GetWidgetContainers()
+{
+	ClearExpiredViewports();
+
+	TSharedPtr<ILevelEditor> LevelEditor = StaticCastSharedPtr<ILevelEditor>(ToolkitHostWeak.Pin());
+
+	if (!LevelEditor.IsValid())
+	{
+		return {};
+	}
+
+	TArray<TSharedPtr<SLevelViewport>> Viewports = LevelEditor->GetViewports();
+
+	TArray<TSharedRef<IFloatingPropertiesWidgetContainer>> Containers;
+	Containers.Reserve(Viewports.Num());
+
+	for (const TSharedPtr<SLevelViewport>& LevelViewport : Viewports)
+	{
+		if (LevelViewport.IsValid())
+		{
+			if (const TSharedRef<FFloatingPropertiesLevelEditorWidgetContainer>* Container = CachedContainers.Find(LevelViewport))
+			{
+				Containers.Add(*Container);
+			}
+			else
+			{
+				TSharedRef<FFloatingPropertiesLevelEditorWidgetContainer> NewContainer = MakeShared<FFloatingPropertiesLevelEditorWidgetContainer>(LevelViewport.ToSharedRef());
+				CachedContainers.Add(LevelViewport, NewContainer);
+				Containers.Add(NewContainer);
+			}
+		}
+	}
+
+	return Containers;
+}
+
+bool FFloatingPropertiesLevelEditorDataProvider::IsWidgetVisibleInContainer(TSharedRef<IFloatingPropertiesWidgetContainer> InContainer) const
+{
+	TSharedPtr<ILevelEditor> LevelEditor = StaticCastSharedPtr<ILevelEditor>(ToolkitHostWeak.Pin());
+
+	if (!LevelEditor.IsValid())
+	{
+		return false;
+	}
+
+	TSharedRef<FFloatingPropertiesLevelEditorWidgetContainer> LevelEditorWidgetContainer = StaticCastSharedRef<FFloatingPropertiesLevelEditorWidgetContainer>(InContainer);
+
+	return LevelEditor->GetActiveViewportInterface() == LevelEditorWidgetContainer->GetLevelViewport();
+}
+
+void FFloatingPropertiesLevelEditorDataProvider::ClearExpiredViewports()
+{
+	for (auto Iter = CachedContainers.CreateIterator(); Iter; ++Iter)
+	{
+		const TWeakPtr<SLevelViewport>& Viewport = Iter.Key();
+
+		if (!Viewport.IsValid())
+		{
+			Iter.RemoveCurrent();
+		}
+	}
+}
diff --git a/Engine/Plugins/Experimental/FloatingProperties/Source/FloatingProperties/Private/LevelEditor/FloatingPropertiesLevelEditorDataProvider.h b/Engine/Plugins/Experimental/FloatingProperties/Source/FloatingProperties/Private/LevelEditor/FloatingPropertiesLevelEditorDataProvider.h
new file mode 100644
--- /dev/null
+++ b/Engine/Plugins/Experimental/FloatingProperties/Source/FloatingProperties/Private/LevelEditor/FloatingPropertiesLevelEditorDataProvider.h
@@ -0,0 +1,29 @@
+// Copyright Epic Games, Inc. All Rights Reserved.
+
+#pragma once
+
+#include "ToolkitHost/FloatingPropertiesToolkitHostDataProvider.h"
+#include "Containers/Map.h"
+#include "Data/IFloatingPropertiesWidgetContainer.h"
+#include "Templates/SharedPointer.h"
+
+class FFloatingPropertiesLevelEditorWidgetContainer;
+class ILevelEditor;
+class SLevelViewport;
+
+class FFloatingPropertiesLevelEditorDataProvider : public FFloatingPropertiesToolkitHostDataProvider
+{
+public:
+	FFloatingPropertiesLevelEditorDataProvider(TSharedRef<ILevelEditor> InLevelEditor);
+	virtual ~FFloatingPropertiesLevelEditorDataProvider() = default;
+
+	//~ Begin IFloatingPropertiesDataProvider
+	virtual TArray<TSharedRef<IFloatingPropertiesWidgetContainer>> GetWidgetContainers() override;
+	virtual bool IsWidgetVisibleInContainer(TSharedRef<IFloatingPropertiesWidgetContainer> InContainer) const override;
+	//~ End IFloatingPropertiesDataProvider
+
+protected:
+	TMap<TWeakPtr<SLevelViewport>, TSharedRef<FFloatingPropertiesLevelEditorWidgetContainer>> CachedContainers;
+
+	void ClearExpiredViewports();
+};
diff --git a/Engine/Plugins/Experimental/FloatingProperties/Source/FloatingProperties/Private/LevelEditor/FloatingPropertiesLevelEditorWidgetContainer.cpp b/Engine/Plugins/Experimental/FloatingProperties/Source/FloatingProperties/Private/LevelEditor/FloatingPropertiesLevelEditorWidgetContainer.cpp
new file mode 100644
--- /dev/null
+++ b/Engine/Plugins/Experimental/FloatingProperties/Source/FloatingProperties/Private/LevelEditor/FloatingPropertiesLevelEditorWidgetContainer.cpp
@@ -0,0 +1,57 @@
+// Copyright Epic Games, Inc. All Rights Reserved.
+
+#include "FloatingPropertiesLevelEditorWidgetContainer.h"
+#include "SLevelViewport.h"
+#include "Widgets/SFloatingPropertiesViewportWidget.h"
+
+FFloatingPropertiesLevelEditorWidgetContainer::FFloatingPropertiesLevelEditorWidgetContainer(TSharedRef<SLevelViewport> InLevelViewport)
+	: LevelViewportWeak(InLevelViewport)
+{
+}
+
+FFloatingPropertiesLevelEditorWidgetContainer::~FFloatingPropertiesLevelEditorWidgetContainer()
+{
+	RemoveWidget();
+}
+
+bool FFloatingPropertiesLevelEditorWidgetContainer::SetWidget(TSharedRef<SFloatingPropertiesViewportWidget> InWidget)
+{
+	if (TSharedPtr<SLevelViewport> LevelViewport = LevelViewportWeak.Pin())
+	{
+		Widget = InWidget;
+		LevelViewport->AddOverlayWidget(InWidget->AsWidget());
+		return true;
+	}
+
+	return false;
+}
+
+TSharedPtr<SFloatingPropertiesViewportWidget> FFloatingPropertiesLevelEditorWidgetContainer::GetWidget() const
+{
+	return Widget;
+}
+
+bool FFloatingPropertiesLevelEditorWidgetContainer::RemoveWidget()
+{
+	if (TSharedPtr<SLevelViewport> LevelViewport = LevelViewportWeak.Pin())
+	{
+		if (Widget.IsValid())
+		{
+			LevelViewport->RemoveOverlayWidget(Widget->AsWidget());
+			Widget.Reset();
+			return true;
+		}
+	}
+
+	return false;
+}
+
+bool FFloatingPropertiesLevelEditorWidgetContainer::IsValid() const
+{
+	return LevelViewportWeak.IsValid();
+}
+
+TSharedPtr<SLevelViewport> FFloatingPropertiesLevelEditorWidgetContainer::GetLevelViewport() const
+{
+	return LevelViewportWeak.Pin();
+}
diff --git a/Engine/Plugins/Experimental/FloatingProperties/Source/FloatingProperties/Private/LevelEditor/FloatingPropertiesLevelEditorWidgetContainer.h b/Engine/Plugins/Experimental/FloatingProperties/Source/FloatingProperties/Private/LevelEditor/FloatingPropertiesLevelEditorWidgetContainer.h
new file mode 100644
--- /dev/null
+++ b/Engine/Plugins/Experimental/FloatingProperties/Source/FloatingProperties/Private/LevelEditor/FloatingPropertiesLevelEditorWidgetContainer.h
@@ -0,0 +1,28 @@
+// Copyright Epic Games, Inc. All Rights Reserved.
+
+#pragma once
+
+#include "Data/IFloatingPropertiesWidgetContainer.h"
+#include "Templates/SharedPointer.h"
+
+class SLevelViewport;
+
+class FFloatingPropertiesLevelEditorWidgetContainer : public IFloatingPropertiesWidgetContainer
+{
+public:
+	FFloatingPropertiesLevelEditorWidgetContainer(TSharedRef<SLevelViewport> InLevelViewport);
+	virtual ~FFloatingPropertiesLevelEditorWidgetContainer();
+
+	//~ Begin IFloatingPropertiesWidgetContainer
+	virtual TSharedPtr<SFloatingPropertiesViewportWidget> GetWidget() const override;
+	virtual bool SetWidget(TSharedRef<SFloatingPropertiesViewportWidget> InWidget) override;
+	virtual bool RemoveWidget() override;
+	virtual bool IsValid() const override;
+	//~ End IFloatingPropertiesWidgetContainer
+
+	TSharedPtr<SLevelViewport> GetLevelViewport() const;
+
+protected:
+	TWeakPtr<SLevelViewport> LevelViewportWeak;
+	TSharedPtr<SFloatingPropertiesViewportWidget> Widget;
+};
diff --git a/Engine/Plugins/Experimental/FloatingProperties/Source/FloatingProperties/Private/LevelEditor/FloatingPropertiesLevelEditorWidgetController.cpp b/Engine/Plugins/Experimental/FloatingProperties/Source/FloatingProperties/Private/LevelEditor/FloatingPropertiesLevelEditorWidgetController.cpp
new file mode 100644
--- /dev/null
+++ b/Engine/Plugins/Experimental/FloatingProperties/Source/FloatingProperties/Private/LevelEditor/FloatingPropertiesLevelEditorWidgetController.cpp
@@ -0,0 +1,151 @@
+// Copyright Epic Games, Inc. All Rights Reserved.
+
+#include "LevelEditor/FloatingPropertiesLevelEditorWidgetController.h"
+#include "Editor.h"
+#include "FloatingPropertiesCommands.h"
+#include "FloatingPropertiesSettings.h"
+#include "LevelEditor.h"
+#include "Editor/EditorEngine.h"
+#include "LevelEditor/FloatingPropertiesLevelEditorDataProvider.h"
+#include "Selection.h"
+
+void FFloatingPropertiesLevelEditorWidgetController::StaticInit()
+{
+	RegisterLevelViewportMenuExtensions();
+	RegisterLevelEditorCommands();
+}
+
+FFloatingPropertiesLevelEditorWidgetController::FFloatingPropertiesLevelEditorWidgetController(TSharedRef<ILevelEditor> InLevelEditor)
+	: FFloatingPropertiesWidgetController(MakeShared<FFloatingPropertiesLevelEditorDataProvider>(InLevelEditor))
+{
+}
+
+FFloatingPropertiesLevelEditorWidgetController::~FFloatingPropertiesLevelEditorWidgetController()
+{
+	UnregisterLevelViewportClientsChanged();
+
+	if (NextTickTimer.IsValid())
+	{
+		FTSTicker::GetCoreTicker().RemoveTicker(NextTickTimer);
+		NextTickTimer.Reset();
+	}	
+}
+
+void FFloatingPropertiesLevelEditorWidgetController::Init()
+{
+	FFloatingPropertiesWidgetController::Init();
+
+	RegisterLevelViewportClientsChanged();
+	OnLevelViewportClientsChanged();
+}
+
+void FFloatingPropertiesLevelEditorWidgetController::RegisterLevelViewportClientsChanged()
+{
+	if (GEditor)
+	{
+		GEditor->OnLevelViewportClientListChanged().AddRaw(this, &FFloatingPropertiesLevelEditorWidgetController::OnLevelViewportClientsChanged);
+	}
+}
+
+void FFloatingPropertiesLevelEditorWidgetController::UnregisterLevelViewportClientsChanged()
+{
+	if (GEditor)
+	{
+		GEditor->OnLevelViewportClientListChanged().RemoveAll(this);
+	}
+}
+
+void FFloatingPropertiesLevelEditorWidgetController::OnLevelViewportClientsChanged()
+{
+	if (NextTickTimer.IsValid())
+	{
+		return;
+	}
+
+	NextTickTimer = FTSTicker::GetCoreTicker().AddTicker(FTickerDelegate::CreateSPLambda(this,
+		[this](float)
+		{
+			OnViewportChange();
+			NextTickTimer.Reset();
+
+			// Do not run again
+			return false;
+		}),
+		/* InDelay - Next frame */ 0.f
+	);
+}
+
+void FFloatingPropertiesLevelEditorWidgetController::RegisterLevelViewportMenuExtensions()
+{
+	static const FName OptionsMenuName("LevelEditor.LevelViewportToolBar.Options");
+	static const FName OptionsMenuSection("LevelViewportViewportOptions2");
+
+	UToolMenus* const ToolMenus = UToolMenus::Get();
+
+	if (!ToolMenus)
+	{
+		return;
+	}
+
+	UToolMenu* const OptionsMenu = ToolMenus->ExtendMenu(OptionsMenuName);
+
+	if (!OptionsMenu)
+	{
+		return;
+	}
+
+	FToolMenuSection* Section = OptionsMenu->FindSection(OptionsMenuSection);
+
+	if (!Section)
+	{
+		return;
+	}
+
+	const FFloatingPropertiesCommands& Commands = FFloatingPropertiesCommands::Get();
+
+	Section->AddEntry(FToolMenuEntry::InitMenuEntry(
+		Commands.ToggleEnabled,
+		TAttribute<FText>(),
+		TAttribute<FText>(),
+		FSlateIcon(FAppStyle::GetAppStyleSetName(), "Icons.Server")
+	));
+}
+
+void FFloatingPropertiesLevelEditorWidgetController::RegisterLevelEditorCommands()
+{
+	FLevelEditorModule& LevelEditorModule = FModuleManager::LoadModuleChecked<FLevelEditorModule>("LevelEditor");
+	const FFloatingPropertiesCommands& Commands = FFloatingPropertiesCommands::Get();
+	
+	LevelEditorModule.GetGlobalLevelEditorActions()->MapAction(
+		Commands.ToggleEnabled,
+		FExecuteAction::CreateLambda(
+			[]()
+			{
+				if (UFloatingPropertiesSettings* Settings = GetMutableDefault<UFloatingPropertiesSettings>())
+				{
+					Settings->bEnabled = !Settings->bEnabled;
+					Settings->SaveConfig();
+					
+					FPropertyChangedEvent ChangeEvent(
+						UFloatingPropertiesSettings::StaticClass()->FindPropertyByName(
+							GET_MEMBER_NAME_CHECKED(UFloatingPropertiesSettings, bEnabled)
+						),
+						EPropertyChangeType::ValueSet
+					);
+						
+					Settings->PostEditChangeProperty(ChangeEvent);
+				}
+			}),
+		FCanExecuteAction(),
+		FIsActionChecked::CreateLambda(
+			[]()
+			{
+				if (const UFloatingPropertiesSettings* Settings = GetDefault<UFloatingPropertiesSettings>())
+				{
+					return Settings->bEnabled;
+				}
+
+				return false;
+			})
+	);
+}
diff --git a/Engine/Plugins/Experimental/FloatingProperties/Source/FloatingProperties/Private/LevelEditor/FloatingPropertiesLevelEditorWidgetController.h b/Engine/Plugins/Experimental/FloatingProperties/Source/FloatingProperties/Private/LevelEditor/FloatingPropertiesLevelEditorWidgetController.h
new file mode 100644
--- /dev/null
+++ b/Engine/Plugins/Experimental/FloatingProperties/Source/FloatingProperties/Private/LevelEditor/FloatingPropertiesLevelEditorWidgetController.h
@@ -0,0 +1,34 @@
+// Copyright Epic Games, Inc. All Rights Reserved.
+
+#pragma once
+
+#include "Data/FloatingPropertiesWidgetController.h"
+#include "Containers/Ticker.h"
+
+class ILevelEditor;
+
+class FFloatingPropertiesLevelEditorWidgetController : public FFloatingPropertiesWidgetController
+{
+public:
+	static void StaticInit();
+
+	FFloatingPropertiesLevelEditorWidgetController(TSharedRef<ILevelEditor> InLevelEditor);
+	virtual ~FFloatingPropertiesLevelEditorWidgetController();
+
+	//~ Begin FFloatingPropertiesWidgetController
+	virtual void Init() override;
+	//~ End FFloatingPropertiesWidgetController
+
+protected:
+	static void RegisterLevelViewportMenuExtensions();
+
+	static void RegisterLevelEditorCommands();
+
+	FTSTicker::FDelegateHandle NextTickTimer;
+
+	void RegisterLevelViewportClientsChanged();
+
+	void UnregisterLevelViewportClientsChanged();
+
+	void OnLevelViewportClientsChanged();
+};
diff --git a/Engine/Plugins/Experimental/FloatingProperties/Source/FloatingProperties/Private/ToolkitHost/FloatingPropertiesToolkitHostDataProvider.cpp b/Engine/Plugins/Experimental/FloatingProperties/Source/FloatingProperties/Private/ToolkitHost/FloatingPropertiesToolkitHostDataProvider.cpp
new file mode 100644
--- /dev/null
+++ b/Engine/Plugins/Experimental/FloatingProperties/Source/FloatingProperties/Private/ToolkitHost/FloatingPropertiesToolkitHostDataProvider.cpp
@@ -0,0 +1,40 @@
+// Copyright Epic Games, Inc. All Rights Reserved.
+
+#include "ToolkitHost/FloatingPropertiesToolkitHostDataProvider.h"
+#include "EditorModeManager.h"
+#include "Toolkits/IToolkitHost.h"
+
+FFloatingPropertiesToolkitHostDataProvider::FFloatingPropertiesToolkitHostDataProvider(TSharedRef<IToolkitHost> InToolkitHost)
+{
+	ToolkitHostWeak = InToolkitHost;
+}
+
+USelection* FFloatingPropertiesToolkitHostDataProvider::GetActorSelection() const
+{
+	if (TSharedPtr<IToolkitHost> ToolkitHost = ToolkitHostWeak.Pin())
+	{
+		return ToolkitHost->GetEditorModeManager().GetSelectedActors();
+	}
+
+	return nullptr;
+}
+
+USelection* FFloatingPropertiesToolkitHostDataProvider::GetComponentSelection() const
+{
+	if (TSharedPtr<IToolkitHost> ToolkitHost = ToolkitHostWeak.Pin())
+	{
+		return ToolkitHost->GetEditorModeManager().GetSelectedComponents();
+	}
+
+	return nullptr;
+}
+
+UWorld* FFloatingPropertiesToolkitHostDataProvider::GetWorld() const
+{
+	if (TSharedPtr<IToolkitHost> ToolkitHost = ToolkitHostWeak.Pin())
+	{
+		return ToolkitHost->GetWorld();
+	}
+
+	return nullptr;
+}
diff --git a/Engine/Plugins/Experimental/FloatingProperties/Source/FloatingProperties/Private/ToolkitHost/FloatingPropertiesToolkitHostDataProvider.h b/Engine/Plugins/Experimental/FloatingProperties/Source/FloatingProperties/Private/ToolkitHost/FloatingPropertiesToolkitHostDataProvider.h
new file mode 100644
--- /dev/null
+++ b/Engine/Plugins/Experimental/FloatingProperties/Source/FloatingProperties/Private/ToolkitHost/FloatingPropertiesToolkitHostDataProvider.h
@@ -0,0 +1,24 @@
+// Copyright Epic Games, Inc. All Rights Reserved.
+
+#pragma once
+
+#include "Data/IFloatingPropertiesDataProvider.h"
+#include "Templates/SharedPointer.h"
+
+class IToolkitHost;
+
+class FFloatingPropertiesToolkitHostDataProvider : public IFloatingPropertiesDataProvider
+{
+public:
+	FFloatingPropertiesToolkitHostDataProvider(TSharedRef<IToolkitHost> InToolkitHost);
+	virtual ~FFloatingPropertiesToolkitHostDataProvider() = default;
+
+	//~ Begin IFloatingPropertiesDataProvider
+	virtual USelection* GetActorSelection() const override;
+	virtual USelection* GetComponentSelection() const override;
+	virtual UWorld* GetWorld() const override;
+	//~ End IFloatingPropertiesDataProvider
+
+protected:
+	TWeakPtr<IToolkitHost> ToolkitHostWeak;
+};
diff --git a/Engine/Plugins/Experimental/FloatingProperties/Source/FloatingProperties/Private/Widgets/CustomPropertyWidgets/SFloatingPropertiesColorPropertyEditor.cpp b/Engine/Plugins/Experimental/FloatingProperties/Source/FloatingProperties/Private/Widgets/CustomPropertyWidgets/SFloatingPropertiesColorPropertyEditor.cpp
new file mode 100644
--- /dev/null
+++ b/Engine/Plugins/Experimental/FloatingProperties/Source/FloatingProperties/Private/Widgets/CustomPropertyWidgets/SFloatingPropertiesColorPropertyEditor.cpp
@@ -0,0 +1,42 @@
+// Copyright Epic Games, Inc. All Rights Reserved.
+
+#include "Widgets/CustomPropertyWidgets/SFloatingPropertiesColorPropertyEditor.h"
+#include "PropertyHandle.h"
+
+void SFloatingPropertiesColorPropertyEditor::PrivateRegisterAttributes(FSlateAttributeInitializer& InInitializer)
+{
+}
+
+TSharedPtr<SWidget> SFloatingPropertiesColorPropertyEditor::CreateWidget(TSharedRef<IPropertyHandle> InPropertyHandle)
+{
+	if (FStructProperty* StructProperty = CastField<FStructProperty>(InPropertyHandle->GetProperty()))
+	{
+		if (StructProperty->Struct == TBaseStructure<FColor>::Get())
+		{
+			return SNew(SFloatingPropertiesColorPropertyEditor, InPropertyHandle);
+		}
+	}
+
+	return nullptr;
+}
+
+void SFloatingPropertiesColorPropertyEditor::Construct(const FArguments& InArgs, TSharedRef<IPropertyHandle> InPropertyHandle)
+{
+	SFloatingPropertiesColorPropertyEditorBase::Construct(
+		SFloatingPropertiesColorPropertyEditorBase::FArguments(),
+		InPropertyHandle
+	);
+}
+
+FLinearColor SFloatingPropertiesColorPropertyEditor::GetColorValue(FProperty* InProperty, UObject* InObject) const
+{
+	FColor Value;
+	InProperty->GetValue_InContainer(InObject, &Value);
+	return Value.ReinterpretAsLinear();
+}
+
+void SFloatingPropertiesColorPropertyEditor::SetColorValue(FProperty* InProperty, UObject* InObject, const FLinearColor& InNewValue)
+{
+	FColor Value = InNewValue.ToFColor(false);
+	InProperty->SetValue_InContainer(InObject, &Value);
+}
diff --git a/Engine/Plugins/Experimental/FloatingProperties/Source/FloatingProperties/Private/Widgets/CustomPropertyWidgets/SFloatingPropertiesColorPropertyEditor.h b/Engine/Plugins/Experimental/FloatingProperties/Source/FloatingProperties/Private/Widgets/CustomPropertyWidgets/SFloatingPropertiesColorPropertyEditor.h
new file mode 100644
--- /dev/null
+++ b/Engine/Plugins/Experimental/FloatingProperties/Source/FloatingProperties/Private/Widgets/CustomPropertyWidgets/SFloatingPropertiesColorPropertyEditor.h
@@ -0,0 +1,24 @@
+// Copyright Epic Games, Inc. All Rights Reserved.
+
+#pragma once
+
+#include "Widgets/CustomPropertyWidgets/SFloatingPropertiesColorPropertyEditorBase.h"
+
+class SFloatingPropertiesColorPropertyEditor : public SFloatingPropertiesColorPropertyEditorBase
+{
+public:
+	SLATE_DECLARE_WIDGET(SFloatingPropertiesColorPropertyEditor, SFloatingPropertiesColorPropertyEditorBase)
+
+	SLATE_BEGIN_ARGS(SFloatingPropertiesColorPropertyEditor) {}
+	SLATE_END_ARGS()
+
+	static TSharedPtr<SWidget> CreateWidget(TSharedRef<IPropertyHandle> InPropertyHandle);
+
+	void Construct(const FArguments& InArgs, TSharedRef<IPropertyHandle> InPropertyHandle);
+
+protected:
+	//~ Begin SFloatingPropertiesColorPropertyEditorBase
+	virtual FLinearColor GetColorValue(FProperty* InProperty, UObject* InObject) const override;
+	virtual void SetColorValue(FProperty* InProperty, UObject* InObject, const FLinearColor& InNewValue) override;
+	//~ End SFloatingPropertiesColorPropertyEditorBase
+};
diff --git a/Engine/Plugins/Experimental/FloatingProperties/Source/FloatingProperties/Private/Widgets/CustomPropertyWidgets/SFloatingPropertiesColorPropertyEditorBase.cpp b/Engine/Plugins/Experimental/FloatingProperties/Source/FloatingProperties/Private/Widgets/CustomPropertyWidgets/SFloatingPropertiesColorPropertyEditorBase.cpp
new file mode 100644
--- /dev/null
+++ b/Engine/Plugins/Experimental/FloatingProperties/Source/FloatingProperties/Private/Widgets/CustomPropertyWidgets/SFloatingPropertiesColorPropertyEditorBase.cpp
@@ -0,0 +1,255 @@
+// Copyright Epic Games, Inc. All Rights Reserved.
+
+#include "Widgets/CustomPropertyWidgets/SFloatingPropertiesColorPropertyEditorBase.h"
+#include "Engine/Engine.h"
+#include "PropertyHandle.h"
+#include "ScopedTransaction.h"
+#include "Widgets/Layout/SBox.h"
+#include "Widgets/Colors/SColorBlock.h"
+#include "Widgets/Colors/SColorPicker.h"
+
+#define LOCTEXT_NAMESPACE "SFloatingPropertiesPropertyWidget"
+
+void SFloatingPropertiesColorPropertyEditorBase::PrivateRegisterAttributes(FSlateAttributeInitializer& InInitializer)
+{
+}
+
+void SFloatingPropertiesColorPropertyEditorBase::Construct(const FArguments& InArgs, TSharedRef<IPropertyHandle> InPropertyHandle)
+{
+	PropertyHandleWeak = InPropertyHandle;
+	bColorPickerOpen = false;
+
+	ChildSlot
+	[
+		SNew(SBox)
+		.WidthOverride(45.f)
+		.HeightOverride(15.f)
+		.HAlign(EHorizontalAlignment::HAlign_Left)
+		[
+			SNew(SColorBlock)
+			.Color(this, &SFloatingPropertiesColorPropertyEditorBase::GetColorValue)
+			.Size(FVector2D(45.f, 15.f))
+			.ShowBackgroundForAlpha(true)
+			.OnMouseButtonDown(this, &SFloatingPropertiesColorPropertyEditorBase::OnClickColorBlock)
+		]
+	];
+}
+
+FLinearColor SFloatingPropertiesColorPropertyEditorBase::GetColorValue() const
+{
+	TSharedPtr<IPropertyHandle> PropertyHandle = PropertyHandleWeak.Pin();
+	FLinearColor Value = FLinearColor::Black;
+
+	if (PropertyHandle.IsValid() && PropertyHandle->IsValidHandle())
+	{
+		if (FProperty* Property = PropertyHandle->GetProperty())
+		{
+			TArray<UObject*> Outers;
+			PropertyHandle->GetOuterObjects(Outers);
+
+			if (!Outers.IsEmpty() && IsValid(Outers[0]))
+			{
+				Value = GetColorValue(Property, Outers[0]);
+			}
+		}
+	}
+
+	return Value;	
+}
+
+void SFloatingPropertiesColorPropertyEditorBase::SetColorValue(const FLinearColor& InNewValue)
+{
+	TSharedPtr<IPropertyHandle> PropertyHandle = PropertyHandleWeak.Pin();
+
+	if (!PropertyHandle.IsValid() || !PropertyHandle->IsValidHandle())
+	{
+		return;
+	}
+
+	FProperty* Property = PropertyHandle->GetProperty();
+
+	if (!Property)
+	{
+		return;
+	}
+
+	TArray<UObject*> Objects;
+	PropertyHandle->GetOuterObjects(Objects);
+
+	for (UObject* Object : Objects)
+	{
+		if (IsValid(Object))
+		{
+			SetColorValue(Property, Object, InNewValue);
+		}
+	}
+}
+
+void SFloatingPropertiesColorPropertyEditorBase::RestoreOriginalColors()
+{
+	if (OriginalValues.IsEmpty())
+	{
+		return;
+	}
+
+	TSharedPtr<IPropertyHandle> PropertyHandle = PropertyHandleWeak.Pin();
+
+	if (!PropertyHandle.IsValid() || !PropertyHandle->IsValidHandle())
+	{
+		return;
+	}
+
+	FProperty* Property = PropertyHandle->GetProperty();
+
+	if (!Property)
+	{
+		return;
+	}
+
+	TArray<UObject*> Objects;
+	PropertyHandle->GetOuterObjects(Objects);
+
+	for (UObject* Object : Objects)
+	{
+		if (IsValid(Object))
+		{
+			if (const FLinearColor* OriginalValue = OriginalValues.Find(Object))
+			{
+				SetColorValue(Property, Object, *OriginalValue);
+			}
+		}
+	}
+
+	OriginalValues.Empty();
+}
+
+FReply SFloatingPropertiesColorPropertyEditorBase::OnClickColorBlock(const FGeometry& MyGeometry, const FPointerEvent& MouseEvent)
+{
+	if (bColorPickerOpen)
+	{
+		return FReply::Unhandled();
+	}
+
+	if (MouseEvent.GetEffectingButton() != EKeys::LeftMouseButton)
+	{
+		return FReply::Unhandled();
+	}
+
+	// Needed for gamma
+	if (!GEngine)
+	{
+		return FReply::Unhandled();
+	}
+
+	OriginalValues.Empty();
+
+	TSharedPtr<IPropertyHandle> PropertyHandle = PropertyHandleWeak.Pin();
+	TArray<UObject*> Objects;
+	bool bValidProperty = false;
+
+	if (PropertyHandle.IsValid() && PropertyHandle->IsValidHandle())
+	{
+		if (FProperty* Property = PropertyHandle->GetProperty())
+		{
+			PropertyHandle->GetOuterObjects(Objects);
+
+			if (!Objects.IsEmpty())
+			{
+				bValidProperty = true;
+				FLinearColor Value;
+				OriginalValues.Reserve(Objects.Num());
+
+				for (UObject* Object : Objects)
+				{
+					if (IsValid(Object))
+					{
+						Value = GetColorValue(Property, Object);
+						OriginalValues.Add(Object, Value);
+					}
+				}
+			}
+		}
+	}
+
+	if (!bValidProperty)
+	{
+		return FReply::Unhandled();
+	}
+
+	FColorPickerArgs PickerArgs;
+	{
+		PickerArgs.bOnlyRefreshOnOk = true;
+		PickerArgs.DisplayGamma = TAttribute<float>::Create(TAttribute<float>::FGetter::CreateUObject(GEngine, &UEngine::GetDisplayGamma));
+		PickerArgs.OnColorCommitted = FOnLinearColorValueChanged::CreateSP(this, &SFloatingPropertiesColorPropertyEditorBase::OnColorPickedCommitted);
+		PickerArgs.OnColorPickerCancelled = FOnColorPickerCancelled::CreateSP(this, &SFloatingPropertiesColorPropertyEditorBase::OnColorPickerCancelled);
+		PickerArgs.OnColorPickerWindowClosed = FOnWindowClosed::CreateSP(this, &SFloatingPropertiesColorPropertyEditorBase::OnColorPickerClosed);
+		PickerArgs.bUseAlpha = true;
+		PickerArgs.InitialColor = GetColorValue();
+		PickerArgs.bOnlyRefreshOnMouseUp = false;
+		PickerArgs.bOnlyRefreshOnOk = false;
+	}
+
+	bColorPickerOpen = OpenColorPicker(PickerArgs);
+
+	if (bColorPickerOpen)
+	{
+		Transaction = MakeShared<FScopedTransaction>(LOCTEXT("Transaction", "Set color value"));
+
+		for (UObject* Object : Objects)
+		{
+			Object->Modify();
+		}
+	}
+
+	return FReply::Handled();
+}
+
+void SFloatingPropertiesColorPropertyEditorBase::OnColorPickedCommitted(FLinearColor InNewColor)
+{
+	TSharedPtr<IPropertyHandle> PropertyHandle = PropertyHandleWeak.Pin();
+
+	if (PropertyHandle.IsValid() && PropertyHandle->IsValidHandle())
+	{
+		PropertyHandle->NotifyPreChange();
+	}
+
+	SetColorValue(InNewColor);
+
+	if (PropertyHandle.IsValid() && PropertyHandle->IsValidHandle())
+	{
+		PropertyHandle->NotifyPostChange(EPropertyChangeType::Interactive);
+	}
+}
+
+void SFloatingPropertiesColorPropertyEditorBase::OnColorPickerCancelled(FLinearColor InOldColor)
+{
+	TSharedPtr<IPropertyHandle> PropertyHandle = PropertyHandleWeak.Pin();
+
+	if (PropertyHandle.IsValid() && PropertyHandle->IsValidHandle())
+	{
+		PropertyHandle->NotifyPreChange();
+	}
+
+	RestoreOriginalColors();
+
+	if (PropertyHandle.IsValid() && PropertyHandle->IsValidHandle())
+	{
+		PropertyHandle->NotifyPostChange(EPropertyChangeType::Interactive);
+	}
+}
+
+void SFloatingPropertiesColorPropertyEditorBase::OnColorPickerClosed(const TSharedRef<SWindow>& InClosedWindow)
+{
+	Transaction.Reset();
+	OriginalValues.Empty();
+	bColorPickerOpen = false;
+
+	TSharedPtr<IPropertyHandle> PropertyHandle = PropertyHandleWeak.Pin();
+
+	if (PropertyHandle.IsValid() && PropertyHandle->IsValidHandle())
+	{
+		PropertyHandle->NotifyFinishedChangingProperties();
+	}
+}
+
+#undef LOCTEXT_NAMESPACE
diff --git a/Engine/Plugins/Experimental/FloatingProperties/Source/FloatingProperties/Private/Widgets/CustomPropertyWidgets/SFloatingPropertiesColorPropertyEditorBase.h b/Engine/Plugins/Experimental/FloatingProperties/Source/FloatingProperties/Private/Widgets/CustomPropertyWidgets/SFloatingPropertiesColorPropertyEditorBase.h
new file mode 100644
--- /dev/null
+++ b/Engine/Plugins/Experimental/FloatingProperties/Source/FloatingProperties/Private/Widgets/CustomPropertyWidgets/SFloatingPropertiesColorPropertyEditorBase.h
@@ -0,0 +1,50 @@
+// Copyright Epic Games, Inc. All Rights Reserved.
+
+#pragma once
+
+#include "Widgets/SCompoundWidget.h"
+#include "Containers/Map.h"
+#include "UObject/WeakObjectPtr.h"
+#include "UObject/WeakObjectPtrTemplates.h"
+
+class IPropertyHandle;
+class FScopedTransaction;
+class SFloatingPropertiesPropertyWidget;
+class UObject;
+
+class SFloatingPropertiesColorPropertyEditorBase : public SCompoundWidget
+{
+public:
+	SLATE_DECLARE_WIDGET(SFloatingPropertiesColorPropertyEditorBase, SCompoundWidget)
+
+	SLATE_BEGIN_ARGS(SFloatingPropertiesColorPropertyEditorBase) {}
+	SLATE_END_ARGS()
+
+	virtual ~SFloatingPropertiesColorPropertyEditorBase() override = default;
+
+	void Construct(const FArguments& InArgs, TSharedRef<IPropertyHandle> InPropertyHandle);
+
+protected:
+	TWeakPtr<IPropertyHandle> PropertyHandleWeak;
+	TSharedPtr<FScopedTransaction> Transaction;
+	TMap<TWeakObjectPtr<UObject>, FLinearColor> OriginalValues;
+	bool bColorPickerOpen;
+
+	virtual FLinearColor GetColorValue(FProperty* InProperty, UObject* InObject) const = 0;
+
+	virtual void SetColorValue(FProperty* InProperty, UObject* InObject, const FLinearColor& InNewValue) = 0;
+
+	virtual FLinearColor GetColorValue() const;
+
+	void SetColorValue(const FLinearColor& InNewValue);
+
+	void RestoreOriginalColors();
+
+	FReply OnClickColorBlock(const FGeometry& MyGeometry, const FPointerEvent& MouseEvent);
+
+	void OnColorPickedCommitted(FLinearColor InNewColor);
+
+	void OnColorPickerCancelled(FLinearColor InOldColor);
+
+	void OnColorPickerClosed(const TSharedRef<SWindow>& InClosedWindow);
+};
diff --git a/Engine/Plugins/Experimental/FloatingProperties/Source/FloatingProperties/Private/Widgets/CustomPropertyWidgets/SFloatingPropertiesLinearColorPropertyEditor.cpp b/Engine/Plugins/Experimental/FloatingProperties/Source/FloatingProperties/Private/Widgets/CustomPropertyWidgets/SFloatingPropertiesLinearColorPropertyEditor.cpp
new file mode 100644
--- /dev/null
+++ b/Engine/Plugins/Experimental/FloatingProperties/Source/FloatingProperties/Private/Widgets/CustomPropertyWidgets/SFloatingPropertiesLinearColorPropertyEditor.cpp
@@ -0,0 +1,41 @@
+// Copyright Epic Games, Inc. All Rights Reserved.
+
+#include "Widgets/CustomPropertyWidgets/SFloatingPropertiesLinearColorPropertyEditor.h"
+#include "PropertyHandle.h"
+
+void SFloatingPropertiesLinearColorPropertyEditor::PrivateRegisterAttributes(FSlateAttributeInitializer& InInitializer)
+{
+}
+
+TSharedPtr<SWidget> SFloatingPropertiesLinearColorPropertyEditor::CreateWidget(TSharedRef<IPropertyHandle> InPropertyHandle)
+{
+	if (FStructProperty* StructProperty = CastField<FStructProperty>(InPropertyHandle->GetProperty()))
+	{
+		if (StructProperty->Struct == TBaseStructure<FColor>::Get())
+		{
+			return SNew(SFloatingPropertiesLinearColorPropertyEditor, InPropertyHandle);
+		}
+	}
+
+	return nullptr;
+}
+
+void SFloatingPropertiesLinearColorPropertyEditor::Construct(const FArguments& InArgs, TSharedRef<IPropertyHandle> InPropertyHandle)
+{
+	SFloatingPropertiesColorPropertyEditorBase::Construct(
+		SFloatingPropertiesColorPropertyEditorBase::FArguments(),
+		InPropertyHandle
+	);
+}
+
+FLinearColor SFloatingPropertiesLinearColorPropertyEditor::GetColorValue(FProperty* InProperty, UObject* InObject) const
+{
+	FLinearColor Value;
+	InProperty->GetValue_InContainer(InObject, &Value);
+	return Value;
+}
+
+void SFloatingPropertiesLinearColorPropertyEditor::SetColorValue(FProperty* InProperty, UObject* InObject, const FLinearColor& InNewValue)
+{
+	InProperty->SetValue_InContainer(InObject, &InNewValue);
+}
diff --git a/Engine/Plugins/Experimental/FloatingProperties/Source/FloatingProperties/Private/Widgets/CustomPropertyWidgets/SFloatingPropertiesLinearColorPropertyEditor.h b/Engine/Plugins/Experimental/FloatingProperties/Source/FloatingProperties/Private/Widgets/CustomPropertyWidgets/SFloatingPropertiesLinearColorPropertyEditor.h
new file mode 100644
--- /dev/null
+++ b/Engine/Plugins/Experimental/FloatingProperties/Source/FloatingProperties/Private/Widgets/CustomPropertyWidgets/SFloatingPropertiesLinearColorPropertyEditor.h
@@ -0,0 +1,24 @@
+// Copyright Epic Games, Inc. All Rights Reserved.
+
+#pragma once
+
+#include "Widgets/CustomPropertyWidgets/SFloatingPropertiesColorPropertyEditorBase.h"
+
+class SFloatingPropertiesLinearColorPropertyEditor : public SFloatingPropertiesColorPropertyEditorBase
+{
+public:
+	SLATE_DECLARE_WIDGET(SFloatingPropertiesLinearColorPropertyEditor, SFloatingPropertiesColorPropertyEditorBase)
+
+	SLATE_BEGIN_ARGS(SFloatingPropertiesLinearColorPropertyEditor) {}
+	SLATE_END_ARGS()
+
+	static TSharedPtr<SWidget> CreateWidget(TSharedRef<IPropertyHandle> InPropertyHandle);
+
+	void Construct(const FArguments& InArgs, TSharedRef<IPropertyHandle> InPropertyHandle);
+
+protected:
+	//~ Begin SFloatingPropertiesColorPropertyEditorBase
+	virtual FLinearColor GetColorValue(FProperty* InProperty, UObject* InObject) const override;
+	virtual void SetColorValue(FProperty* InProperty, UObject* InObject, const FLinearColor& InNewValue) override;
+	//~ End SFloatingPropertiesColorPropertyEditorBase
+};
diff --git a/Engine/Plugins/Experimental/FloatingProperties/Source/FloatingProperties/Private/Widgets/FloatingPropertiesDragOperation.cpp b/Engine/Plugins/Experimental/FloatingProperties/Source/FloatingProperties/Private/Widgets/FloatingPropertiesDragOperation.cpp
new file mode 100644
--- /dev/null
+++ b/Engine/Plugins/Experimental/FloatingProperties/Source/FloatingProperties/Private/Widgets/FloatingPropertiesDragOperation.cpp
@@ -0,0 +1,58 @@
+// Copyright Epic Games, Inc. All Rights Reserved.
+
+#include "Widgets/FloatingPropertiesDragOperation.h"
+#include "Framework/Application/SlateApplication.h"
+#include "Widgets/SFloatingPropertiesPropertyWidget.h"
+#include "Widgets/SFloatingPropertiesViewportWidget.h"
+
+#define LOCTEXT_NAMESPACE "FloatingPropertiesDragOperation"
+
+FFloatingPropertiesDragOperation::FFloatingPropertiesDragOperation(TSharedPtr<SFloatingPropertiesPropertyWidget> InPropertyWidget)
+{
+	PropertyWidgetWeak = InPropertyWidget;
+	MouseStartPosition = FSlateApplication::Get().GetCursorPos();
+	MouseToWidgetDelta = MouseStartPosition - InPropertyWidget->GetTickSpaceGeometry().GetAbsolutePosition();
+}
+
+FFloatingPropertiesDragOperation::~FFloatingPropertiesDragOperation()
+{
+	TSharedPtr<SFloatingPropertiesPropertyWidget> PropertyWidget = PropertyWidgetWeak.Pin();
+
+	if (!PropertyWidget.IsValid())
+	{
+		return;
+	}
+
+	TSharedPtr<SFloatingPropertiesViewportWidget> ViewportWidget = PropertyWidget->GetViewportWidget();
+
+	if (!ViewportWidget.IsValid())
+	{
+		return;
+	}
+
+	ViewportWidget->OnPropertyDragComplete(PropertyWidget.ToSharedRef());
+}
+
+void FFloatingPropertiesDragOperation::OnDragged(const FDragDropEvent& DragDropEvent)
+{
+	FDragDropOperation::OnDragged(DragDropEvent);
+
+	TSharedPtr<SFloatingPropertiesPropertyWidget> PropertyWidget = PropertyWidgetWeak.Pin();
+
+	if (!PropertyWidget.IsValid())
+	{
+		return;
+	}
+
+	TSharedPtr<SFloatingPropertiesViewportWidget> ViewportWidget = PropertyWidget->GetViewportWidget();
+
+	if (!ViewportWidget.IsValid())
+	{
+		return;
+	}
+
+	ViewportWidget->OnPropertyDragUpdate(PropertyWidget.ToSharedRef(), MouseToWidgetDelta, MouseStartPosition, 
+		FSlateApplication::Get().GetCursorPos());
+}
+
+#undef LOCTEXT_NAMESPACE
diff --git a/Engine/Plugins/Experimental/FloatingProperties/Source/FloatingProperties/Private/Widgets/FloatingPropertiesDragOperation.h b/Engine/Plugins/Experimental/FloatingProperties/Source/FloatingProperties/Private/Widgets/FloatingPropertiesDragOperation.h
new file mode 100644
--- /dev/null
+++ b/Engine/Plugins/Experimental/FloatingProperties/Source/FloatingProperties/Private/Widgets/FloatingPropertiesDragOperation.h
@@ -0,0 +1,25 @@
+// Copyright Epic Games, Inc. All Rights Reserved.
+
+#pragma once
+
+#include "Input/DragAndDrop.h"
+#include "Math/Vector2D.h"
+#include "Templates/SharedPointer.h"
+
+class SFloatingPropertiesPropertyWidget;
+
+class FFloatingPropertiesDragOperation : public FDragDropOperation
+{
+public:
+	DRAG_DROP_OPERATOR_TYPE(FFloatingPropertiesDragOperation, FDragDropOperation)
+
+	FFloatingPropertiesDragOperation(TSharedPtr<SFloatingPropertiesPropertyWidget> InPropertyWidget);
+	virtual ~FFloatingPropertiesDragOperation() override;
+
+	virtual void OnDragged(const class FDragDropEvent& DragDropEvent) override;
+
+protected:
+	TWeakPtr<SFloatingPropertiesPropertyWidget> PropertyWidgetWeak;
+	FVector2f MouseStartPosition;
+	FVector2f MouseToWidgetDelta;
+};
diff --git a/Engine/Plugins/Experimental/FloatingProperties/Source/FloatingProperties/Private/Widgets/SFloatingPropertiesDragButton.cpp b/Engine/Plugins/Experimental/FloatingProperties/Source/FloatingProperties/Private/Widgets/SFloatingPropertiesDragButton.cpp
new file mode 100644
--- /dev/null
+++ b/Engine/Plugins/Experimental/FloatingProperties/Source/FloatingProperties/Private/Widgets/SFloatingPropertiesDragButton.cpp
@@ -0,0 +1,50 @@
+// Copyright Epic Games, Inc. All Rights Reserved.
+
+#include "Widgets/SFloatingPropertiesDragButton.h"
+#include "FloatingPropertiesDragOperation.h"
+#include "Styling/AppStyle.h"
+#include "Widgets/SFloatingPropertiesPropertyWidget.h"
+
+#define LOCTEXT_NAMESPACE "SFloatingPropertiesDragButton"
+
+void SFloatingPropertiesDragButton::PrivateRegisterAttributes(FSlateAttributeInitializer& InInitializer)
+{
+}
+
+void SFloatingPropertiesDragButton::Construct(const FArguments& InArgs, TSharedRef<SFloatingPropertiesPropertyWidget> InPropertyWidget)
+{
+	PropertyWidgetWeak = InPropertyWidget;
+
+	SButton::Construct(
+		SButton::FArguments()
+			.ButtonStyle(FAppStyle::Get(), "SimpleButton")
+			.OnClicked(InArgs._OnClicked)
+			[
+				InArgs._Content.Widget
+			]
+	);
+}
+
+FReply SFloatingPropertiesDragButton::OnMouseButtonDown(const FGeometry& MyGeometry, const FPointerEvent& MouseEvent)
+{
+	FReply Reply = SButton::OnMouseButtonDown(MyGeometry, MouseEvent);
+
+	if (PropertyWidgetWeak.IsValid())
+	{
+		return Reply.DetectDrag(SharedThis(this), EKeys::LeftMouseButton);
+	}
+
+	return Reply;
+}
+
+FReply SFloatingPropertiesDragButton::OnDragDetected(const FGeometry& MyGeometry, const FPointerEvent& MouseEvent)
+{
+	if (TSharedPtr<SFloatingPropertiesPropertyWidget> PropertyWidget = PropertyWidgetWeak.Pin())
+	{
+		return FReply::Handled().BeginDragDrop(MakeShared<FFloatingPropertiesDragOperation>(PropertyWidget.ToSharedRef()));
+	}
+
+	return FReply::Unhandled();
+}
+
+#undef LOCTEXT_NAMESPACE
diff --git a/Engine/Plugins/Experimental/FloatingProperties/Source/FloatingProperties/Private/Widgets/SFloatingPropertiesDragButton.h b/Engine/Plugins/Experimental/FloatingProperties/Source/FloatingProperties/Private/Widgets/SFloatingPropertiesDragButton.h
new file mode 100644
--- /dev/null
+++ b/Engine/Plugins/Experimental/FloatingProperties/Source/FloatingProperties/Private/Widgets/SFloatingPropertiesDragButton.h
@@ -0,0 +1,31 @@
+// Copyright Epic Games, Inc. All Rights Reserved.
+
+#pragma once
+
+#include "Widgets/Input/SButton.h"
+
+class SFloatingPropertiesPropertyWidget;
+
+class SFloatingPropertiesDragButton : public SButton
+{
+public:
+	SLATE_DECLARE_WIDGET(SFloatingPropertiesDragButton, SButton)
+
+	SLATE_BEGIN_ARGS(SFloatingPropertiesDragButton)
+		{}
+		SLATE_DEFAULT_SLOT(FArguments, Content)
+		SLATE_EVENT(FOnClicked, OnClicked)
+	SLATE_END_ARGS()
+
+	virtual ~SFloatingPropertiesDragButton() override = default;
+
+	void Construct(const FArguments& InArgs, TSharedRef<SFloatingPropertiesPropertyWidget> InPropertyWidget);
+
+	//~ Begin SWidget
+	virtual FReply OnMouseButtonDown(const FGeometry& MyGeometry, const FPointerEvent& MouseEvent) override;
+	virtual FReply OnDragDetected(const FGeometry& MyGeometry, const FPointerEvent& MouseEvent) override;
+	//~ End SWidget
+
+protected:
+	TWeakPtr<SFloatingPropertiesPropertyWidget> PropertyWidgetWeak;
+};
diff --git a/Engine/Plugins/Experimental/FloatingProperties/Source/FloatingProperties/Private/Widgets/SFloatingPropertiesDragContainer.cpp b/Engine/Plugins/Experimental/FloatingProperties/Source/FloatingProperties/Private/Widgets/SFloatingPropertiesDragContainer.cpp
new file mode 100644
--- /dev/null
+++ b/Engine/Plugins/Experimental/FloatingProperties/Source/FloatingProperties/Private/Widgets/SFloatingPropertiesDragContainer.cpp
@@ -0,0 +1,42 @@
+// Copyright Epic Games, Inc. All Rights Reserved.
+
+#include "Widgets/SFloatingPropertiesDragContainer.h"
+#include "Widgets/FloatingPropertiesDragOperation.h"
+
+#define LOCTEXT_NAMESPACE "SFloatingPropertiesDragContainer"
+
+void SFloatingPropertiesDragContainer::PrivateRegisterAttributes(FSlateAttributeInitializer& InInitializer)
+{
+}
+
+void SFloatingPropertiesDragContainer::Construct(const FArguments& InArgs, TSharedRef<SFloatingPropertiesPropertyWidget> InPropertyWidget)
+{
+	PropertyWidgetWeak = InPropertyWidget;
+
+	ChildSlot
+	[
+		InArgs._Content.Widget
+	];
+}
+
+FReply SFloatingPropertiesDragContainer::OnMouseButtonDown(const FGeometry& MyGeometry, const FPointerEvent& MouseEvent)
+{
+	if (PropertyWidgetWeak.IsValid())
+	{
+		return FReply::Handled().DetectDrag(SharedThis(this), EKeys::LeftMouseButton);
+	}
+
+	return FReply::Unhandled();
+}
+
+FReply SFloatingPropertiesDragContainer::OnDragDetected(const FGeometry& MyGeometry, const FPointerEvent& MouseEvent)
+{
+	if (TSharedPtr<SFloatingPropertiesPropertyWidget> PropertyWidget = PropertyWidgetWeak.Pin())
+	{
+		return FReply::Handled().BeginDragDrop(MakeShared<FFloatingPropertiesDragOperation>(PropertyWidget.ToSharedRef()));
+	}
+
+	return FReply::Unhandled();
+}
+
+#undef LOCTEXT_NAMESPACE
diff --git a/Engine/Plugins/Experimental/FloatingProperties/Source/FloatingProperties/Private/Widgets/SFloatingPropertiesDragContainer.h b/Engine/Plugins/Experimental/FloatingProperties/Source/FloatingProperties/Private/Widgets/SFloatingPropertiesDragContainer.h
new file mode 100644
--- /dev/null
+++ b/Engine/Plugins/Experimental/FloatingProperties/Source/FloatingProperties/Private/Widgets/SFloatingPropertiesDragContainer.h
@@ -0,0 +1,30 @@
+// Copyright Epic Games, Inc. All Rights Reserved.
+
+#pragma once
+
+#include "Widgets/SCompoundWidget.h"
+
+class SFloatingPropertiesPropertyWidget;
+
+class SFloatingPropertiesDragContainer : public SCompoundWidget
+{
+public:
+	SLATE_DECLARE_WIDGET(SFloatingPropertiesDragContainer, SCompoundWidget)
+
+	SLATE_BEGIN_ARGS(SFloatingPropertiesDragContainer)
+		{}
+		SLATE_DEFAULT_SLOT(FArguments, Content)
+	SLATE_END_ARGS()
+
+	virtual ~SFloatingPropertiesDragContainer() override = default;
+
+	void Construct(const FArguments& InArgs, TSharedRef<SFloatingPropertiesPropertyWidget> InPropertyWidget);
+
+	//~ Begin SWidget
+	virtual FReply OnMouseButtonDown(const FGeometry& MyGeometry, const FPointerEvent& MouseEvent) override;
+	virtual FReply OnDragDetected(const FGeometry& MyGeometry, const FPointerEvent& MouseEvent) override;
+	//~ End SWidget
+
+protected:
+	TWeakPtr<SFloatingPropertiesPropertyWidget> PropertyWidgetWeak;
+};
diff --git a/Engine/Plugins/Experimental/FloatingProperties/Source/FloatingProperties/Private/Widgets/SFloatingPropertiesPropertyWidget.cpp b/Engine/Plugins/Experimental/FloatingProperties/Source/FloatingProperties/Private/Widgets/SFloatingPropertiesPropertyWidget.cpp
new file mode 100644
--- /dev/null
+++ b/Engine/Plugins/Experimental/FloatingProperties/Source/FloatingProperties/Private/Widgets/SFloatingPropertiesPropertyWidget.cpp
@@ -0,0 +1,577 @@
+// Copyright Epic Games, Inc. All Rights Reserved.
+
+#include "Widgets/SFloatingPropertiesPropertyWidget.h"
+#include "Algo/Reverse.h"
+#include "ClassIconFinder.h"
+#include "Components/ActorComponent.h"
+#include "DetailLayoutBuilder.h"
+#include "DetailWidgetRow.h"
+#include "Dialogs/Dialogs.h"
+#include "FloatingPropertiesModule.h"
+#include "FloatingPropertiesSettings.h"
+#include "IDetailTreeNode.h"
+#include "Framework/Commands/UIAction.h"
+#include "Framework/MultiBox/MultiBoxBuilder.h"
+#include "GameFramework/Actor.h"
+#include "Internationalization/Text.h"
+#include "Misc/MessageDialog.h"
+#include "PropertyHandle.h"
+#include "Styling/AppStyle.h"
+#include "UObject/UnrealType.h"
+#include "Widgets/Images/SImage.h"
+#include "Widgets/Input/SEditableTextBox.h"
+#include "Widgets/Input/SMenuAnchor.h"
+#include "Widgets/Layout/SBorder.h"
+#include "Widgets/Layout/SBox.h"
+#include "Widgets/SFloatingPropertiesDragButton.h"
+#include "Widgets/SFloatingPropertiesDragContainer.h"
+#include "Widgets/SFloatingPropertiesViewportWidget.h"
+#include "Widgets/SNullWidget.h"
+#include "Widgets/Text/STextBlock.h"
+
+#define LOCTEXT_NAMESPACE "SFloatingPropertiesPropertyWidget"
+
+void SFloatingPropertiesPropertyWidget::PrivateRegisterAttributes(FSlateAttributeInitializer& InInitializer)
+{
+}
+
+void SFloatingPropertiesPropertyWidget::Construct(const FArguments& InArgs, TSharedRef<SFloatingPropertiesViewportWidget> InViewportWidget,
+	const FFloatingPropertiesPropertyTypes& InPropertyTypes)
+{
+	ViewportWidgetWeak = InViewportWidget;
+	PropertyHandleWeak = InPropertyTypes.PropertyHandle;
+
+	SetVisibility(EVisibility::Visible);
+
+	check(InPropertyTypes.PropertyHandle.IsValid());
+	check(InPropertyTypes.Property);
+	check(InPropertyTypes.TreeNode.IsValid());
+
+	// Use first selected object to get property icons
+	TArray<UObject*> OuterObjects;
+	InPropertyTypes.PropertyHandle->GetOuterObjects(OuterObjects);
+	check(!OuterObjects.IsEmpty());
+
+	bool bAllowSaveValue = !InPropertyTypes.Property->IsA<FBoolProperty>();
+
+	const FSlateBrush* Icon = nullptr;
+
+	if (AActor* Actor = Cast<AActor>(OuterObjects[0]))
+	{
+		Icon = FClassIconFinder::FindIconForActor(Actor);
+	}
+	else if (UActorComponent* Component = Cast<UActorComponent>(OuterObjects[0]))
+	{
+		Icon = FSlateIconFinder::FindIconBrushForClass(Component->GetClass(), TEXT("SCS.NativeComponent"));
+	}
+	else
+	{
+		checkNoEntry();
+	}
+
+	{
+		TSharedPtr<IPropertyHandle> LoopHandle = InPropertyTypes.PropertyHandle;
+		UClass* OwningClass = nullptr;
+		FString PropertyPath = "";
+		TArray<FText> PropertyDisplayNames;
+
+		do
+		{
+			if (FProperty* Property = LoopHandle->GetProperty())
+			{
+				PropertyPath = Property->GetName() + (PropertyPath.Len() > 0 ? FString(TEXT(".")) : FString()) + PropertyPath;
+				PropertyDisplayNames.Add(Property->GetDisplayNameText());
+				OwningClass = Property->GetOwnerClass();
+			}
+
+			LoopHandle = LoopHandle->GetParentHandle();
+		}
+		while (!OwningClass && LoopHandle.IsValid());
+
+		ClassProperty = {TSoftClassPtr<UObject>(OwningClass),	PropertyPath};
+
+		const FSlateFontInfo DetailFont = IDetailLayoutBuilder::GetDetailFont();
+
+		Algo::Reverse(PropertyDisplayNames);
+
+		PropertyNameWidget = SNew(STextBlock)
+			.Font(DetailFont)
+			.Text(FText::Join(INVTEXT("."), PropertyDisplayNames));
+	}
+
+	if (FStructProperty* StructProperty = CastField<FStructProperty>(InPropertyTypes.Property))
+	{
+		const FFloatingPropertiesModule::FCreateStructPropertyValueWidgetDelegate* Delegate
+			= FFloatingPropertiesModule::Get().GetStructPropertyValueWidgetDelegate(StructProperty->Struct);
+
+		if (Delegate && Delegate->IsBound())
+		{
+			if (TSharedPtr<SWidget> Widget = Delegate->Execute(InPropertyTypes.PropertyHandle.ToSharedRef()))
+			{
+				PropertyValueWidget = Widget;
+			}
+		}
+
+		if (!PropertyValueWidget.IsValid())
+		{
+			if (TSharedPtr<class IDetailPropertyRow> TreeRow = InPropertyTypes.TreeNode->GetRow())
+			{
+				TSharedPtr<SWidget> RowNameWidget;
+				TSharedPtr<SWidget> RowValueWidget;
+				FDetailWidgetRow Row;
+
+				TreeRow->GetDefaultWidgets(RowNameWidget, RowValueWidget, Row, /* bAddWidgetDecoration */ false);
+
+				if (RowValueWidget.IsValid())
+				{
+					PropertyValueWidget = RowValueWidget.ToSharedRef();
+					bAllowSaveValue = false;
+				}				
+			}
+		}
+	}
+
+	if (!PropertyValueWidget.IsValid())
+	{
+		PropertyValueWidget = InPropertyTypes.PropertyHandle->CreatePropertyValueWidget(false);
+	}
+
+	if (!PropertyValueWidget.IsValid())
+	{
+		PropertyValueWidget = SNullWidget::NullWidget;
+	}
+
+	TSharedRef<SWidget> IconWidget = SNew(SImage)
+		.DesiredSizeOverride(FVector2D(12.f, 12.f))
+		.Image(Icon);
+
+	if (!bAllowSaveValue)
+	{
+		IconWidget = SNew(SFloatingPropertiesDragContainer, SharedThis(this))
+			[
+				IconWidget
+			];
+	}
+	else
+	{
+		IconWidget = SNew(SFloatingPropertiesDragButton, SharedThis(this))
+			.OnClicked(this, &SFloatingPropertiesPropertyWidget::OnPropertyButtonClicked)
+			[
+				SAssignNew(PropertyAnchor, SMenuAnchor)
+					.OnGetMenuContent(this, &SFloatingPropertiesPropertyWidget::MakePropertyMenu)
+					[
+						IconWidget
+					]
+			];
+	}
+
+	ChildSlot
+	[
+		SNew(SBorder)
+		.Padding(1.f, 1.f, 1.f, 1.f)
+		.BorderImage(FAppStyle::GetBrush("ToolTip.Background"))
+		.BorderBackgroundColor(FLinearColor(1.f, 1.f, 1.f, 0.75f))
+		.OnMouseButtonDown_Lambda([](const FGeometry&, const FPointerEvent&) { return FReply::Handled(); })
+		.OnMouseButtonUp_Lambda([](const FGeometry&, const FPointerEvent&) { return FReply::Handled(); })
+		.OnMouseDoubleClick_Lambda([](const FGeometry&, const FPointerEvent&) { return FReply::Handled(); })
+		[
+			SNew(SHorizontalBox)
+
+			+ SHorizontalBox::Slot()
+			.AutoWidth()
+			.Padding(0.f, 0.f, 0.f, 0.f)
+			[
+				SNew(SBox)
+				.WidthOverride(24.f)
+				.HAlign(EHorizontalAlignment::HAlign_Center)
+				.VAlign(EVerticalAlignment::VAlign_Center)
+				[
+					IconWidget
+				]
+			]
+
+			+ SHorizontalBox::Slot()
+			.AutoWidth()
+			.Padding(3.f, 1.f, 0.f, 0.f)
+			[
+				SNew(SBox)
+				.WidthOverride(this, &SFloatingPropertiesPropertyWidget::GetPropertyNameOverrideSize)
+				.HAlign(EHorizontalAlignment::HAlign_Fill)
+				.VAlign(EVerticalAlignment::VAlign_Center)
+				[
+					SNew(SFloatingPropertiesDragContainer, SharedThis(this))
+					[
+						PropertyNameWidget.ToSharedRef()
+					]
+				]
+			]
+
+			+ SHorizontalBox::Slot()
+			.AutoWidth()
+			.Padding(3.f, 1.f, 0.f, 0.f)
+			[
+				SNew(SBox)
+				.WidthOverride(this, &SFloatingPropertiesPropertyWidget::GetPropertyValueOverrideSize)
+				.VAlign(EVerticalAlignment::VAlign_Center)
+				[
+					PropertyValueWidget.ToSharedRef()
+				]
+			]
+		]
+	];
+}
+
+float SFloatingPropertiesPropertyWidget::GetPropertyNameWidgetWidth() const
+{
+	if (PropertyNameWidget.IsValid())
+	{
+		return PropertyNameWidget->GetDesiredSize().X;
+	}
+
+	return -1.f;
+}
+
+float SFloatingPropertiesPropertyWidget::GetPropertyValueWidgetWidth() const
+{
+	if (PropertyValueWidget.IsValid())
+	{
+		return PropertyValueWidget->GetDesiredSize().X;
+	}
+
+	return -1.f;
+}
+
+FFloatingPropertiesClassProperties* SFloatingPropertiesPropertyWidget::GetSavedValues(bool bInEnsure) const
+{
+	if (bInEnsure)
+	{
+		return &GetMutableDefault<UFloatingPropertiesSettings>()->SavedValues.FindOrAdd(ClassProperty);
+	}
+
+	if (UFloatingPropertiesSettings* FloatingPropertiesSettings = GetMutableDefault<UFloatingPropertiesSettings>())
+	{
+		if (FFloatingPropertiesClassProperties* ClassProperties = FloatingPropertiesSettings->SavedValues.Find(ClassProperty))
+		{
+			return ClassProperties;
+		}
+	}
+
+	return nullptr;
+}
+
+void SFloatingPropertiesPropertyWidget::SaveConfig()
+{
+	if (UFloatingPropertiesSettings* FloatingPropertiesSettings = GetMutableDefault<UFloatingPropertiesSettings>())
+	{
+		FloatingPropertiesSettings->SaveConfig();
+	}
+}
+
+void SFloatingPropertiesPropertyWidget::AddSavedValue()
+{
+	TSharedPtr<IPropertyHandle> PropertyHandle = PropertyHandleWeak.Pin();
+
+	if (!PropertyHandle.IsValid())
+	{
+		return;
+	}
+
+	FString NewName = "Value";
+
+	if (const FFloatingPropertiesClassProperties* SavedValues = GetSavedValues(/* bInEnsure */ false))
+	{
+		FString Value;
+		PropertyHandle->GetValueAsFormattedString(Value);
+
+		if (SavedValues->Properties.FindKey(Value))
+		{
+			FMessageDialog::Open(
+				EAppMsgCategory::Error,
+				EAppMsgType::Ok,
+				LOCTEXT("AlreadySavedErrorBody", "This value has already been saved."),
+				LOCTEXT("AlreadySavedErrorTitle", "Already Saved")
+			);
+
+			return;
+		}
+
+		NewName = UFloatingPropertiesSettings::FindUniqueName(*SavedValues, NewName);
+	}
+
+	TSharedRef<SEditableTextBox> NameInput = SNew(SEditableTextBox)
+		.HintText(LOCTEXT("NewValueHint", "Saved Value Name"))
+		.MinDesiredWidth(100.f)
+		.Text(FText::FromString(NewName));
+
+	SGenericDialogWidget::OpenDialog(
+		LOCTEXT("NewValueTitle", "New Value Name"),
+		NameInput,
+		SGenericDialogWidget::FArguments(),
+		/* Modal */ true
+	);
+
+	const FString NewValueName = NameInput->GetText().ToString().TrimStartAndEnd();
+
+	if (NewValueName.IsEmpty())
+	{
+		UE_LOG(LogFloatingProperties, Error, TEXT("Attempted to save a value with no name."));
+		return;
+	}
+
+	AddSavedValue(NewValueName);
+}
+
+void SFloatingPropertiesPropertyWidget::AddSavedValue(const FString& InValueName)
+{
+	TSharedPtr<IPropertyHandle> PropertyHandle = PropertyHandleWeak.Pin();
+
+	if (!PropertyHandle.IsValid())
+	{
+		return;
+	}
+
+	FProperty* Property = PropertyHandle->GetProperty();
+
+	if (!Property)
+	{
+		return;
+	}
+
+	FFloatingPropertiesClassProperties* SavedValues = GetSavedValues(/* bInEnsure */ true);
+
+	if (!SavedValues)
+	{
+		return;
+	}
+
+	const FString ValueName = UFloatingPropertiesSettings::FindUniqueName(*SavedValues, InValueName);
+
+	FString Value;
+	PropertyHandle->GetValueAsFormattedString(Value);
+
+	SavedValues->Properties.Add(ValueName, Value);
+	SaveConfig();
+}
+
+void SFloatingPropertiesPropertyWidget::ApplySavedValue(FString InValueName)
+{
+	TSharedPtr<IPropertyHandle> PropertyHandle = PropertyHandleWeak.Pin();
+
+	if (!PropertyHandle.IsValid())
+	{
+		return;
+	}
+
+	const FFloatingPropertiesClassProperties* SavedValues = GetSavedValues(/* bInEnsure */ false);
+
+	if (!SavedValues)
+	{
+		return;
+	}
+
+	const FString* SavedValuePtr = SavedValues->Properties.Find(InValueName);
+
+	if (!SavedValuePtr)
+	{
+		return;
+	}
+
+	PropertyHandle->SetValueFromFormattedString(*SavedValuePtr);
+}
+
+bool SFloatingPropertiesPropertyWidget::IsValueActive(FString InValueName) const
+{
+	TSharedPtr<IPropertyHandle> PropertyHandle = PropertyHandleWeak.Pin();
+
+	if (!PropertyHandle.IsValid())
+	{
+		return false;
+	}
+
+	const FFloatingPropertiesClassProperties* SavedValues = GetSavedValues(/* bInEnsure */ false);
+
+	if (!SavedValues)
+	{
+		return false;
+	}
+
+	const FString* ValuePtr = SavedValues->Properties.Find(InValueName);
+
+	if (!ValuePtr)
+	{
+		return false;
+	}
+
+	FString Value;
+	PropertyHandle->GetValueAsFormattedString(Value);
+
+	return ValuePtr->Equals(Value);
+}
+
+void SFloatingPropertiesPropertyWidget::RemoveAllValues()
+{
+	TSharedPtr<IPropertyHandle> PropertyHandle = PropertyHandleWeak.Pin();
+
+	if (!PropertyHandle.IsValid())
+	{
+		return;
+	}
+
+	FFloatingPropertiesClassProperties* SavedValues = GetSavedValues(/* bInEnsure */ false);
+
+	if (!SavedValues)
+	{
+		return;
+	}
+
+	SavedValues->Properties.Empty();
+	SaveConfig();
+}
+
+bool SFloatingPropertiesPropertyWidget::CanSaveValue(FString InValue) const
+{
+	TSharedPtr<IPropertyHandle> PropertyHandle = PropertyHandleWeak.Pin();
+
+	if (!PropertyHandle.IsValid())
+	{
+		return false;
+	}
+
+	const FFloatingPropertiesClassProperties* SavedValues = GetSavedValues(/* bInEnsure */ false);
+
+	if (!SavedValues)
+	{
+		return true;
+	}
+
+	return !SavedValues->Properties.FindKey(InValue);
+}
+
+FReply SFloatingPropertiesPropertyWidget::OnPropertyButtonClicked()
+{
+	TSharedPtr<IPropertyHandle> PropertyHandle = PropertyHandleWeak.Pin();
+
+	if (!PropertyHandle.IsValid())
+	{
+		return FReply::Unhandled();
+	}
+
+	if (!PropertyAnchor.IsValid())
+	{
+		return FReply::Unhandled();
+	}
+
+	if (!PropertyAnchor->IsOpen())
+	{
+		PropertyAnchor->SetIsOpen(/* bInIsOpen */ true, /* bFocusMenu */ true);
+	}
+
+	return FReply::Handled();
+}
+
+TSharedRef<SWidget> SFloatingPropertiesPropertyWidget::MakePropertyMenu()
+{
+	TSharedPtr<IPropertyHandle> PropertyHandle = PropertyHandleWeak.Pin();
+
+	if (!PropertyHandle.IsValid())
+	{
+		return SNullWidget::NullWidget;
+	}
+
+	FProperty* Property = PropertyHandle->GetProperty();
+
+	if (!Property)
+	{
+		return SNullWidget::NullWidget;
+	}
+
+	if (Property->IsA<FBoolProperty>())
+	{
+		return SNullWidget::NullWidget;
+	}
+
+	FString Value;
+	PropertyHandle->GetValueAsFormattedString(Value);
+
+	FUIAction ResetToDefaultAction
+	(
+		FExecuteAction::CreateSP(PropertyHandle.Get(), &IPropertyHandle::ResetToDefault),
+		FCanExecuteAction::CreateSP(PropertyHandle.Get(), &IPropertyHandle::CanResetToDefault)
+	);
+
+	FUIAction AddSavedValueAction
+	(
+		FExecuteAction::CreateSP(this, &SFloatingPropertiesPropertyWidget::AddSavedValue),
+		FCanExecuteAction::CreateSP(this, &SFloatingPropertiesPropertyWidget::CanSaveValue, Value)
+	);
+
+	FMenuBuilder MenuBuilder(/* bInShouldCloseWindowAfterMenuSelection */ true, nullptr, nullptr);
+
+	MenuBuilder.AddMenuEntry
+	(
+		LOCTEXT("PropertyResetToDefaultLabel", "Reset To Default"),
+		LOCTEXT("PropertyResetToDefaultToolTip", "Reset to the property to its default value."),
+		FSlateIcon(),
+		ResetToDefaultAction,
+		NAME_None,
+		EUserInterfaceActionType::Button
+	);
+
+	MenuBuilder.AddMenuEntry
+	(
+		LOCTEXT("PropertySaveValueLabel", "Save Value"),
+		LOCTEXT("PropertySaveValueToolTip", "Save the current value."),
+		FSlateIcon(),
+		AddSavedValueAction,
+		NAME_None,
+		EUserInterfaceActionType::Button
+	);
+
+	const FFloatingPropertiesClassProperties* SavedValues = GetSavedValues(/* bInEnsure */ false);
+
+	if (SavedValues && !SavedValues->Properties.IsEmpty())
+	{
+		MenuBuilder.AddMenuSeparator();
+
+		for (const TPair<FString, FString>& Pair : SavedValues->Properties)
+		{
+			FUIAction SetSavedValue
+			(
+				FExecuteAction::CreateSP(this, &SFloatingPropertiesPropertyWidget::ApplySavedValue, Pair.Key),
+				FCanExecuteAction(),
+				FIsActionChecked::CreateSP(this, &SFloatingPropertiesPropertyWidget::IsValueActive, Pair.Key)
+			);
+
+			MenuBuilder.AddMenuEntry
+			(
+				FText::FromString(Pair.Key),
+				FText::Format(LOCTEXT("ValueTooltip", "Value: {0}"), FText::FromString(Pair.Value)),
+				FSlateIcon(),
+				SetSavedValue,
+				NAME_None,
+				EUserInterfaceActionType::Check
+			);
+		}
+
+		MenuBuilder.AddMenuSeparator();
+
+		FUIAction RemoveAllValues
+		(
+			FExecuteAction::CreateSP(this, &SFloatingPropertiesPropertyWidget::RemoveAllValues)
+		);
+
+		MenuBuilder.AddMenuEntry
+		(
+			LOCTEXT("PropertyRemoveAllValuesLabel", "Remove All Values"),
+			LOCTEXT("PropertyRemoveAllValuesTooltio", "Remove all the saved values."),
+			FSlateIcon(),
+			RemoveAllValues,
+			NAME_None,
+			EUserInterfaceActionType::Button
+		);
+	}
+
+	return MenuBuilder.MakeWidget();
+}
+
+#undef LOCTEXT_NAMESPACE
diff --git a/Engine/Plugins/Experimental/FloatingProperties/Source/FloatingProperties/Private/Widgets/SFloatingPropertiesPropertyWidget.h b/Engine/Plugins/Experimental/FloatingProperties/Source/FloatingProperties/Private/Widgets/SFloatingPropertiesPropertyWidget.h
new file mode 100644
--- /dev/null
+++ b/Engine/Plugins/Experimental/FloatingProperties/Source/FloatingProperties/Private/Widgets/SFloatingPropertiesPropertyWidget.h
@@ -0,0 +1,84 @@
+// Copyright Epic Games, Inc. All Rights Reserved.
+
+#pragma once
+
+#include "Widgets/SCompoundWidget.h"
+#include "FloatingPropertiesSettings.h"
+#include "Templates/SharedPointer.h"
+#include "Types/SlateStructs.h"
+
+class IPropertyHandle;
+class SFloatingPropertiesViewportWidget;
+class SMenuAnchor;
+class STextBlock;
+struct FFloatingPropertiesClassProperties;
+struct FFloatingPropertiesPropertyTypes;
+
+class SFloatingPropertiesPropertyWidget : public SCompoundWidget
+{
+public:
+	SLATE_DECLARE_WIDGET(SFloatingPropertiesPropertyWidget, SCompoundWidget)
+
+	SLATE_BEGIN_ARGS(SFloatingPropertiesPropertyWidget) {}
+	SLATE_END_ARGS()
+
+	virtual ~SFloatingPropertiesPropertyWidget() override = default;
+
+	void Construct(const FArguments& InArgs, TSharedRef<SFloatingPropertiesViewportWidget> InViewportWidget,
+		const FFloatingPropertiesPropertyTypes& InPropertyTypes);
+
+	TSharedPtr<SFloatingPropertiesViewportWidget> GetViewportWidget() const { return ViewportWidgetWeak.Pin(); }
+
+	TSharedPtr<IPropertyHandle> GetPropertyHandle() const { return PropertyHandleWeak.Pin(); }
+
+	const FFloatingPropertiesClassProperty& GetClassProperty() const { return ClassProperty; }
+
+	float GetPropertyNameWidgetWidth() const;
+
+	float GetPropertyValueWidgetWidth() const;
+
+	void SetPropertyNameOverrideSize(FOptionalSize InSize) { PropertyNameOverrideSize = InSize; }
+
+	void SetPropertyValueOverrideSize(FOptionalSize InSize) { PropertyValueOverrideSize = InSize; }
+
+	void SaveConfig();
+
+protected:
+	TWeakPtr<SFloatingPropertiesViewportWidget> ViewportWidgetWeak;
+
+	TWeakPtr<IPropertyHandle> PropertyHandleWeak;
+
+	TSharedPtr<SMenuAnchor> PropertyAnchor;
+
+	TSharedPtr<SWidget> PropertyNameWidget;
+
+	TSharedPtr<SWidget> PropertyValueWidget;
+
+	FFloatingPropertiesClassProperty ClassProperty;
+
+	FOptionalSize PropertyNameOverrideSize;
+
+	FOptionalSize PropertyValueOverrideSize;
+
+	FFloatingPropertiesClassProperties* GetSavedValues(bool bInEnsure) const;
+
+	FOptionalSize GetPropertyNameOverrideSize() const { return PropertyNameOverrideSize; }
+
+	FOptionalSize GetPropertyValueOverrideSize() const { return PropertyValueOverrideSize; }
+
+	void AddSavedValue();
+
+	void AddSavedValue(const FString& InValueName);
+
+	void ApplySavedValue(FString InValueName);
+
+	bool IsValueActive(FString InValueName) const;
+
+	bool CanSaveValue(FString InValue) const;
+
+	void RemoveAllValues();
+
+	FReply OnPropertyButtonClicked();
+
+	TSharedRef<SWidget> MakePropertyMenu();
+};
\ No newline at end of file
diff --git a/Engine/Plugins/Experimental/FloatingProperties/Source/FloatingProperties/Private/Widgets/SFloatingPropertiesViewportWidget.cpp b/Engine/Plugins/Experimental/FloatingProperties/Source/FloatingProperties/Private/Widgets/SFloatingPropertiesViewportWidget.cpp
new file mode 100644
--- /dev/null
+++ b/Engine/Plugins/Experimental/FloatingProperties/Source/FloatingProperties/Private/Widgets/SFloatingPropertiesViewportWidget.cpp
@@ -0,0 +1,530 @@
+// Copyright Epic Games, Inc. All Rights Reserved.
+
+#include "Widgets/SFloatingPropertiesViewportWidget.h"
+#include "Components/ActorComponent.h"
+#include "Data/FloatingPropertiesPropertyNodeContainer.h"
+#include "Data/FloatingPropertiesPropertyNode.h"
+#include "Data/FloatingPropertiesSnapMetrics.h"
+#include "Data/FloatingPropertiesWidgetController.h"
+#include "Framework/Application/SlateApplication.h"
+#include "GameFramework/Actor.h"
+#include "PropertyHandle.h"
+#include "Widgets/Layout/SBox.h"
+#include "Widgets/SCanvas.h"
+#include "Widgets/SFloatingPropertiesPropertyWidget.h"
+
+#define LOCTEXT_NAMESPACE "SFloatingPropertiesViewportWidget"
+
+namespace UE::FloatingProperties::Private
+{
+	constexpr float MinPropertyValueWidth = 100.f;
+
+	void HomogenizeStackWidth(TSharedRef<FFloatingPropertiesPropertyNode> InNode)
+	{
+		TArray<TSharedRef<FFloatingPropertiesPropertyNode>> StackNodes = InNode->GetNodeStack();
+
+		float MaxNameWidth = -1.f;
+
+		float MaxValueWidth = -1.f;
+
+		for (const TSharedRef<FFloatingPropertiesPropertyNode>& StackNode : StackNodes)
+		{
+			TSharedRef<SFloatingPropertiesPropertyWidget> AttachedWidget = StackNode->GetPropertyWidget();
+			bool bAllowNarrowWidget = false;
+
+			if (TSharedPtr<IPropertyHandle> PropertyHandle = AttachedWidget->GetPropertyHandle())
+			{
+				if (FProperty* Property = PropertyHandle->GetProperty())
+				{
+					/**
+					 * Bools only get a really tiny edit widget which does not require a good width to display properly.
+					 * Other types, such as text or numeric types, when their edit widgets attempt to use their preferred
+					 * size, display really tiny, unusable widgets. Thus, for non-bool properties a minimum width is
+					 * specified so that they display nicely.
+					 */
+					if (Property->IsA<FBoolProperty>())
+					{
+						bAllowNarrowWidget = true;
+					}
+
+					// Struct properties are responsible for managing their own width.
+					else if (Property->IsA<FStructProperty>())
+					{
+						bAllowNarrowWidget = true;
+					}
+				}
+			}
+
+			const float NameWidth = AttachedWidget->GetPropertyNameWidgetWidth();
+			float ValueWidth = AttachedWidget->GetPropertyValueWidgetWidth();
+
+			if (!bAllowNarrowWidget)
+			{
+				ValueWidth = FMath::Max(ValueWidth, MinPropertyValueWidth);
+			}			
+
+			MaxNameWidth = FMath::Max(MaxNameWidth, NameWidth);
+			MaxValueWidth = FMath::Max(MaxValueWidth, ValueWidth);
+		}
+
+		for (const TSharedRef<FFloatingPropertiesPropertyNode>& StackNode : StackNodes)
+		{
+			if (MaxNameWidth > 0.f)
+			{
+				StackNode->GetPropertyWidget()->SetPropertyNameOverrideSize(FOptionalSize(MaxNameWidth));
+			}
+			else
+			{
+				StackNode->GetPropertyWidget()->SetPropertyNameOverrideSize(FOptionalSize());
+			}
+
+			if (MaxValueWidth > 0.f)
+			{
+				StackNode->GetPropertyWidget()->SetPropertyValueOverrideSize(FOptionalSize(MaxValueWidth));
+			}
+			else
+			{
+				StackNode->GetPropertyWidget()->SetPropertyValueOverrideSize(FOptionalSize());
+			}
+		}
+	}
+}
+
+void SFloatingPropertiesViewportWidget::PrivateRegisterAttributes(FSlateAttributeInitializer& InInitializer)
+{
+}
+
+void SFloatingPropertiesViewportWidget::Construct(const FArguments& InArgs)
+{
+	SetVisibility(EVisibility::SelfHitTestInvisible);
+
+	Properties = MakeShared<FloatingPropertiesPropertyNodeContainer>(SharedThis(this));
+
+	ChildSlot
+	[
+		SNew(SBox)
+		.Padding(2.f)
+		.HAlign(EHorizontalAlignment::HAlign_Fill)
+		.VAlign(EVerticalAlignment::VAlign_Fill)
+		.Visibility_Lambda([IsVisibleCallback = InArgs._IsVisible]()
+			{
+				const bool bIsVisible = IsVisibleCallback.IsBound() ? IsVisibleCallback.Execute() : true;
+
+				return bIsVisible ? EVisibility::SelfHitTestInvisible : EVisibility::Collapsed;
+			})
+		[
+			SAssignNew(FieldContainer, SCanvas)
+			.Visibility(EVisibility::SelfHitTestInvisible)
+		]
+	];
+
+	Reset();
+}
+
+void SFloatingPropertiesViewportWidget::BuildWidget(AActor* InSelectedActor, const FFloatingPropertiesPropertyList& InActorProperties,
+	UActorComponent* InSelectedComponent, const FFloatingPropertiesPropertyList& InComponentProperties)
+{
+	Reset();
+
+	const bool bValidActor = IsValid(InSelectedActor);
+	const bool bValidComponent = IsValid(InSelectedComponent);
+
+	if (!bValidActor && !bValidComponent)
+	{
+		return;
+	}
+
+	// Get property list for this selection
+	TArray<FFloatingPropertiesPropertyTypes> PropertyList;
+
+	if (bValidActor)
+	{
+		PropertyList.Append(InActorProperties.Properties);
+	}
+
+	if (bValidComponent)
+	{
+		PropertyList.Append(InComponentProperties.Properties);
+	}
+	
+	Properties->Reserve(PropertyList.Num());
+
+	TMap<FFloatingPropertiesClassProperty, TSharedRef<FFloatingPropertiesPropertyNode>> ClassPropertyToNode;
+	ClassPropertyToNode.Reserve(PropertyList.Num());
+
+	// Build class property list/widgets
+	for (const FFloatingPropertiesPropertyTypes& Property : PropertyList)
+	{
+		if (!Property.Property || !Property.PropertyHandle.IsValid() || Property.PropertyHandle->GetNumOuterObjects() == 0)
+		{
+			continue;
+		}
+
+		TSharedRef<SFloatingPropertiesPropertyWidget> NewWidget =
+			SNew(SFloatingPropertiesPropertyWidget, SharedThis(this), Property);
+
+		TSharedRef<FFloatingPropertiesPropertyNode> Node = Properties->AddWidget(NewWidget);
+
+		FFloatingPropertiesClassProperty ClassProperty = Node->GetPropertyWidget()->GetClassProperty();
+
+		ClassPropertyToNode.Add(ClassProperty, Node);
+
+		FieldContainer->AddSlot()
+		.Size_Lambda([WidgetWeak = Node->GetPropertyWidget().ToWeakPtr()]() -> FVector2D
+			{
+				if (TSharedPtr<SFloatingPropertiesPropertyWidget> Widget = WidgetWeak.Pin())
+				{
+					return Widget->GetDesiredSize();
+				}
+
+				return FVector2D::ZeroVector;
+			})
+		.Position_Lambda([NodeWeak = Node.ToWeakPtr()]() -> FVector2D
+			{
+				if (TSharedPtr<FFloatingPropertiesPropertyNode> Node = NodeWeak.Pin())
+				{
+					return static_cast<FVector2D>(Node->GetCachedPosition());
+				}
+
+				return FVector2D::ZeroVector;
+			})
+		[
+			Node->GetPropertyWidget()
+		];
+	}
+
+	// Load data
+	const UFloatingPropertiesSettings* FloatingPropertiesSettings = GetDefault<UFloatingPropertiesSettings>();
+
+	for (const TSharedRef<FFloatingPropertiesPropertyNode>& Node : Properties->GetNodes())
+	{
+		FFloatingPropertiesClassProperty ClassProperty = Node->GetPropertyWidget()->GetClassProperty();
+
+		if (const FFloatingPropertiesClassPropertyPosition* SavedPropertyPosition = FloatingPropertiesSettings->PropertyPositions.Find(ClassProperty))
+		{
+			Node->SetPropertyPositionDirect(*SavedPropertyPosition);
+		}
+
+		if (const FFloatingPropertiesClassPropertyAnchor* SavedPropertyAnchor = FloatingPropertiesSettings->PropertyAnchors.Find(ClassProperty))
+		{
+			if (SavedPropertyAnchor->ParentProperty.IsValid())
+			{
+				if (const TSharedRef<FFloatingPropertiesPropertyNode>* ParentNode = ClassPropertyToNode.Find(SavedPropertyAnchor->ParentProperty))
+				{
+					if (!Node->HasSibling(*ParentNode))
+					{
+						Node->SetParent((*ParentNode)->GetStackLeafMostNode());
+					}
+				}
+
+				if (const TSharedRef<FFloatingPropertiesPropertyNode>* ChildNode = ClassPropertyToNode.Find(SavedPropertyAnchor->ChildProperty))
+				{
+					if (!Node->HasSibling(*ChildNode))
+					{
+						Node->SetChild((*ChildNode)->GetStackRootNode());
+					}
+				}
+			}
+		}
+	}
+
+	// Build the default stack at the bottom left
+	TArray<TSharedRef<FFloatingPropertiesPropertyNode>> UnpositionedProperties;
+	UnpositionedProperties.Reserve(Properties->GetNodes().Num());
+
+	for (const TSharedRef<FFloatingPropertiesPropertyNode>& Node : Properties->GetNodes())
+	{
+		if (!Node->HasParent() && Node->GetPropertyPosition().OnDefaultStack())
+		{
+			UnpositionedProperties.Add(Node);
+		}
+	}
+
+	if (UnpositionedProperties.Num() > 1)
+	{
+		TSharedPtr<FFloatingPropertiesPropertyNode> PreviousNode;
+
+		for (int32 Index = 0; Index < UnpositionedProperties.Num(); ++Index)
+		{
+			TSharedRef<FFloatingPropertiesPropertyNode> Node = UnpositionedProperties[Index];
+
+			if (PreviousNode.IsValid())
+			{
+				Node->SetParent(PreviousNode.ToSharedRef());
+			}
+
+			PreviousNode = Node;
+		}
+	}
+
+	// Nodes cannot be laid out here because we still don't know our size. Hysteresis powers combine!
+}
+
+TSharedRef<SWidget> SFloatingPropertiesViewportWidget::AsWidget()
+{
+	return SharedThis(this);
+}
+
+FVector2f SFloatingPropertiesViewportWidget::GetDraggableArea() const
+{
+	static const FVector2f BorderArea = FVector2f(4.f, 4.f);
+
+	return GetTickSpaceGeometry().GetLocalSize() - BorderArea;
+}
+
+void SFloatingPropertiesViewportWidget::OnArrangeChildren(const FGeometry& AllottedGeometry, FArrangedChildren& ArrangedChildren) const
+{
+	SCompoundWidget::OnArrangeChildren(AllottedGeometry, ArrangedChildren);
+
+	if (!FMath::IsNearlyZero(GetTickSpaceGeometry().GetLocalSize().X))
+	{
+		CheckForViewportSizeChange(AllottedGeometry.GetLocalSize());
+		EnsurePositions(/* bInInvalidate */ !bInitialPositionsSet);
+		bInitialPositionsSet = true;
+	}
+}
+
+void SFloatingPropertiesViewportWidget::OnPropertyDragUpdate(TSharedRef<SFloatingPropertiesPropertyWidget> InPropertyWidget, 
+	const FVector2f& InMouseToWidgetDelta, const FVector2f& InMouseStartPosition, const FVector2f& InMouseCurrentPosition)
+{
+	TSharedPtr<FFloatingPropertiesPropertyNode> DraggedNodePtr = Properties->FindNodeForWidget(InPropertyWidget);
+
+	if (!DraggedNodePtr.IsValid())
+	{
+		return;
+	}
+
+	TSharedRef<FFloatingPropertiesPropertyNode> DraggedNode = DraggedNodePtr.ToSharedRef();
+
+	const FVector2f DraggableArea = GetDraggableArea();
+	constexpr float MinAxisValue = UE_SMALL_NUMBER;
+
+	if (DraggableArea.X < MinAxisValue || DraggableArea.Y < MinAxisValue)
+	{
+		return;
+	}
+
+	using namespace UE::FloatingProperties::Private;
+
+	// Work out the space we have to play with
+	const FVector2f DragAreaPosition = GetTickSpaceGeometry().GetAbsolutePosition();
+	const FVector2f DraggedToMousePosition = InMouseCurrentPosition - InMouseToWidgetDelta;
+	const FVector2f DraggedNodeStackSize = DraggedNode->GetStackSize();
+
+	FFloatingPropertiesSnapMetrics DraggedNodeMetrics[2] = {
+		{
+			DraggedNode,
+			DraggedToMousePosition,
+			0.f,
+			EFloatingPropertiesSnapType::AttachAsChild
+		},
+		{
+			DraggedNode,
+			DraggedToMousePosition + FVector2f(0.f, DraggedNodeStackSize.Y),
+			0.f,
+			EFloatingPropertiesSnapType::AttachAsParent
+		}
+	};
+
+	// If we're breaking away from the stack
+	const bool bIsolateProperty = FSlateApplication::Get().GetModifierKeys().IsAltDown();
+
+	if (bIsolateProperty)
+	{
+		DraggedNode->RemoveFromStack();
+		DraggedNode->SetPropertyPosition(DraggableArea, DraggedNodeMetrics[FFloatingPropertiesSnapMetrics::TopLeft].Position - DragAreaPosition);
+		HomogenizeStackWidth(DraggedNode);
+		return;
+	}
+
+	TSharedPtr<FFloatingPropertiesPropertyNode> CurrentParent = DraggedNode->GetParent();
+
+	// Check to see if we're still locked to our current stack
+	if (CurrentParent.IsValid())
+	{
+		const FVector2f OriginalParentSnapPoint = FFloatingPropertiesSnapMetrics::GetSnapPosition(
+			CurrentParent.ToSharedRef(), 
+			EFloatingPropertiesSnapType::AttachAsParent
+		);
+
+		const float DistanceFromParentSq = (DraggedNodeMetrics[FFloatingPropertiesSnapMetrics::TopLeft].Position - OriginalParentSnapPoint).SizeSquared();
+
+		if (DistanceFromParentSq < FFloatingPropertiesSnapMetrics::SnapBreakDistanceSq)
+		{
+			return;
+		}
+	}
+
+	TSet<TSharedRef<FFloatingPropertiesPropertyNode>> DraggedStackSet(DraggedNode->GetNodeStack());
+
+	TArray<FFloatingPropertiesSnapMetrics> SnapMetrics;
+	SnapMetrics.Reserve((Properties->GetNodes().Num() - 1) * 2);
+
+	// Gather a list of snap candidates.
+	for (const TSharedRef<FFloatingPropertiesPropertyNode>& Node : Properties->GetNodes())
+	{
+		// We're only interested nodes outside of the dragged stack.
+		if (DraggedStackSet.Contains(Node))
+		{
+			continue;
+		}
+
+		// Match the bottom left of the node stack to the top left of the dragged node
+		if (!Node->HasChild())
+		{
+			SnapMetrics.Add(FFloatingPropertiesSnapMetrics::Make(
+				Node, 
+				DraggedNodeMetrics[FFloatingPropertiesSnapMetrics::TopLeft].Position, 
+				EFloatingPropertiesSnapType::AttachAsParent
+			));
+		}
+
+		// Match the top left of the node stack to the bottom left of the dragged node
+		if (!Node->HasParent())
+		{
+			SnapMetrics.Add(FFloatingPropertiesSnapMetrics::Make(
+				Node, 
+				DraggedNodeMetrics[FFloatingPropertiesSnapMetrics::BottomLeft].Position, 
+				EFloatingPropertiesSnapType::AttachAsChild
+			));
+		}
+	}
+
+	// Work out the closest
+	int32 ClosestMetricsIndex = INDEX_NONE;
+
+	for (int32 Index = 0; Index < SnapMetrics.Num(); ++Index)
+	{
+		const FFloatingPropertiesSnapMetrics& SnapMetric = SnapMetrics[Index];
+
+		if (SnapMetric.DistanceSq < FFloatingPropertiesSnapMetrics::SnapDistanceSq &&
+			(ClosestMetricsIndex == INDEX_NONE || SnapMetric.DistanceSq < SnapMetrics[ClosestMetricsIndex].DistanceSq))
+		{
+			ClosestMetricsIndex = Index;
+		}
+	}
+
+	// Work out if we need to change node relationships
+	TSharedPtr<FFloatingPropertiesPropertyNode> NewParent;
+	TSharedPtr<FFloatingPropertiesPropertyNode> NewChild;
+	bool bNodeRelationshipChanged = false;
+
+	if (ClosestMetricsIndex != INDEX_NONE)
+	{
+		FFloatingPropertiesSnapMetrics& SnapMetric = SnapMetrics[ClosestMetricsIndex];
+
+		switch (SnapMetric.AttachType)
+		{
+			case EFloatingPropertiesSnapType::AttachAsParent:
+				NewParent = SnapMetric.Node;
+				break;
+
+			case EFloatingPropertiesSnapType::AttachAsChild:
+				NewChild = SnapMetric.Node;
+				break;
+		}
+	}
+
+	// Attach new parent, if any change.
+	if (CurrentParent != NewParent)
+	{
+		bNodeRelationshipChanged = true;
+
+		if (CurrentParent.IsValid())
+		{
+			// Enter free drag
+			DraggedNode->RemoveParent();
+
+			// Update the old parent's stack width
+			HomogenizeStackWidth(CurrentParent.ToSharedRef());
+		}
+
+		if (NewParent.IsValid())
+		{
+			// Make sure we don't have a parent
+			if (TSharedPtr<FFloatingPropertiesPropertyNode> NewParentOldChild = NewParent->RemoveChild())
+			{
+				HomogenizeStackWidth(NewParentOldChild.ToSharedRef());
+			}
+
+			// Snap to parent
+			DraggedNode->SetParent(NewParent.ToSharedRef());
+		}
+	}
+
+	// Attach new child, if there is one.
+	if (NewChild.IsValid())
+	{
+		bNodeRelationshipChanged = true;
+
+		// Make sure we don't have a child
+		if (TSharedPtr<FFloatingPropertiesPropertyNode> NewChildOldParent = NewChild->RemoveParent())
+		{
+			HomogenizeStackWidth(NewChildOldParent.ToSharedRef());
+		}
+
+		DraggedNode->GetStackLeafMostNode()->SetChild(NewChild.ToSharedRef());
+	}
+
+	// If we now don't have a parent, update the position for dragging.
+	if (!DraggedNode->HasParent())
+	{
+		// Snap to sides of the screen
+		const FVector2f AvailableDraggableArea = DraggableArea - DraggedNode->GetStackSize();
+		FVector2f DraggedToPosition = DraggedNodeMetrics[FFloatingPropertiesSnapMetrics::TopLeft].Position - DragAreaPosition;
+
+		FFloatingPropertiesSnapMetrics::SnapToDraggableArea(AvailableDraggableArea, DraggedToPosition);
+
+		DraggedNode->SetPropertyPosition(DraggableArea, DraggedToPosition);
+	}
+
+	// If we changed node attachments anywhere, make sure we're all the same width.
+	if (bNodeRelationshipChanged)
+	{
+		HomogenizeStackWidth(DraggedNode);
+	}
+}
+
+void SFloatingPropertiesViewportWidget::OnPropertyDragComplete(TSharedRef<SFloatingPropertiesPropertyWidget> InPropertyWidget)
+{
+	InPropertyWidget->SaveConfig();
+}
+
+void SFloatingPropertiesViewportWidget::Reset()
+{
+	bInitialPositionsSet = false;
+	Properties->Empty();
+	FieldContainer->ClearChildren();
+}
+
+void SFloatingPropertiesViewportWidget::CheckForViewportSizeChange(const FVector2f& InViewportSize) const
+{
+	const FIntPoint CurrentRenderSize = {
+		FMath::RoundToInt(InViewportSize.X),
+		FMath::RoundToInt(InViewportSize.Y)
+	};
+
+	if (CurrentRenderSize != LastRenderSize)
+	{
+		Properties->InvalidateAllPositions();
+		LastRenderSize = CurrentRenderSize;
+	}
+}
+
+void SFloatingPropertiesViewportWidget::EnsurePositions(bool bInInvalidate) const
+{
+	using namespace UE::FloatingProperties::Private;
+
+	for (const TSharedRef<FFloatingPropertiesPropertyNode>& Node : Properties->GetNodes())
+	{
+		if (!Node->HasParent())
+		{
+			if (Node->UpdateStackPositions(bInInvalidate) == EFloatingPropertiesUpdateResult::Updated)
+			{
+				HomogenizeStackWidth(Node);
+			}
+		}
+	}
+}
+
+#undef LOCTEXT_NAMESPACE
diff --git a/Engine/Plugins/Experimental/FloatingProperties/Source/FloatingProperties/Private/Widgets/SFloatingPropertiesViewportWidget.h b/Engine/Plugins/Experimental/FloatingProperties/Source/FloatingProperties/Private/Widgets/SFloatingPropertiesViewportWidget.h
new file mode 100644
--- /dev/null
+++ b/Engine/Plugins/Experimental/FloatingProperties/Source/FloatingProperties/Private/Widgets/SFloatingPropertiesViewportWidget.h
@@ -0,0 +1,72 @@
+// Copyright Epic Games, Inc. All Rights Reserved.
+
+#pragma once
+
+#include "Widgets/SCompoundWidget.h"
+#include "Containers/Array.h"
+#include "Data/FloatingPropertiesPropertyNodeContainer.h"
+#include "Data/FloatingPropertiesWidgetController.h"
+#include "FloatingPropertiesSettings.h"
+#include "Templates/SharedPointer.h"
+
+class AActor;
+class IPropertyHandle;
+class SCanvas;
+class SFloatingPropertiesPropertyWidget;
+class UActorComponent;
+struct FFloatingPropertiesPropertyList;
+
+class SFloatingPropertiesViewportWidget : public SCompoundWidget
+{
+public:
+	SLATE_DECLARE_WIDGET(SFloatingPropertiesViewportWidget, SCompoundWidget)
+
+	SLATE_BEGIN_ARGS(SFloatingPropertiesViewportWidget)
+		{}
+		SLATE_EVENT(FIsVisibleDelegate, IsVisible)
+	SLATE_END_ARGS()
+
+	virtual ~SFloatingPropertiesViewportWidget() override = default;
+
+	void Construct(const FArguments& InArgs);
+
+	//~ Begin SWidget
+	virtual void OnArrangeChildren(const FGeometry& AllottedGeometry, FArrangedChildren& ArrangedChildren) const override;
+	//~ End SWidget
+
+	void BuildWidget(AActor* InSelectedActor, const FFloatingPropertiesPropertyList& InActorProperties,
+		UActorComponent* InSelectedComponent, const FFloatingPropertiesPropertyList& InComponentProperties);
+
+	TSharedRef<SWidget> AsWidget();
+
+	FVector2f GetDraggableArea() const;
+
+	/**
+	 * The bottom of the stack attaches to the top of the closest snap node if it's above the node.
+	 *
+	 * The top of the stack attaches to the bottom of the closest snap node if it's below the node.
+	 *
+	 * Attaching to the bottom of another stack is temporary. If the node continues to be dragged, after
+	 * SnapBreakDistanceSq from the snapped-to node, it will break off the stack it attached to and
+	 * continue to be dragged.
+	 */
+	void OnPropertyDragUpdate(TSharedRef<SFloatingPropertiesPropertyWidget> InPropertyWidget, const FVector2f& InMouseToWidgetDelta,
+		const FVector2f& InMouseStartPosition, const FVector2f& InMouseCurrentPosition);
+
+	void OnPropertyDragComplete(TSharedRef<SFloatingPropertiesPropertyWidget> InPropertyWidget);
+
+protected:
+	TSharedPtr<SCanvas> FieldContainer;
+
+	TSharedPtr<FloatingPropertiesPropertyNodeContainer> Properties;
+
+	mutable FIntPoint LastRenderSize;
+
+	mutable bool bInitialPositionsSet;
+
+	void Reset();
+
+	void CheckForViewportSizeChange(const FVector2f& InViewportSize) const;
+
+	void EnsurePositions(bool bInInvalidate) const;
+};
diff --git a/Engine/Plugins/Experimental/FloatingProperties/Source/FloatingProperties/Public/Data/FloatingPropertiesWidgetController.h b/Engine/Plugins/Experimental/FloatingProperties/Source/FloatingProperties/Public/Data/FloatingPropertiesWidgetController.h
new file mode 100644
--- /dev/null
+++ b/Engine/Plugins/Experimental/FloatingProperties/Source/FloatingProperties/Public/Data/FloatingPropertiesWidgetController.h
@@ -0,0 +1,76 @@
+// Copyright Epic Games, Inc. All Rights Reserved.
+
+#pragma once
+
+#include "Components/ActorComponent.h"
+#include "Containers/Array.h"
+#include "Data/IFloatingPropertiesDataProvider.h"
+#include "Delegates/Delegate.h"
+#include "Delegates/DelegateCombinations.h"
+#include "GameFramework/Actor.h"
+#include "Templates/SharedPointer.h"
+
+class IDetailTreeNode;
+class IPropertyHandle;
+class IPropertyRowGenerator;
+class SFloatingPropertiesViewportWidget;
+class SFloatingPropertiesViewportWidget;
+class SWidget;
+class USelection;
+class UTypedElementSelectionSet;
+
+struct FFloatingPropertiesPropertyTypes
+{
+	TSharedPtr<IPropertyHandle> PropertyHandle;
+	FProperty* Property;
+	TSharedPtr<IDetailTreeNode> TreeNode;
+};
+
+struct FFloatingPropertiesPropertyList
+{
+	TSharedPtr<IPropertyRowGenerator> Generator;
+	TArray<FFloatingPropertiesPropertyTypes> Properties;
+};
+
+DECLARE_DELEGATE_RetVal(bool, FIsVisibleDelegate);
+
+class FLOATINGPROPERTIES_API FFloatingPropertiesWidgetController : public TSharedFromThis<FFloatingPropertiesWidgetController>
+{
+public:
+	FFloatingPropertiesWidgetController(TSharedRef<IFloatingPropertiesDataProvider> InDataProvider);
+
+	virtual ~FFloatingPropertiesWidgetController();
+
+	virtual void Init();
+
+protected:
+	TSharedPtr<IFloatingPropertiesDataProvider> DataProvider;
+
+	TWeakObjectPtr<AActor> SelectedActorWeak;
+
+	TWeakObjectPtr<UActorComponent> SelectedComponentWeak;
+
+	FFloatingPropertiesPropertyList ActorProperties;
+
+	FFloatingPropertiesPropertyList ComponentProperties;
+
+	TSharedRef<SFloatingPropertiesViewportWidget> CreateWidget(FIsVisibleDelegate InVisibilityCallback);
+
+	void UpdatePropertiesForObject(UObject* InObject, FFloatingPropertiesPropertyList& OutPropertyList);
+
+	void RebuildWidgets(const FFloatingPropertiesPropertyList& InActorProperties, const FFloatingPropertiesPropertyList& InComponentProperties);
+
+	void RemoveWidgets();
+
+	void OnViewportChange();
+
+	void RegisterSelectionChanged();
+
+	void UnregisterSelectionChanged();
+
+	void OnSelectionChange(const UTypedElementSelectionSet* InSelectionSet);
+
+	void OnSelectionChange();
+
+	void OnSelectionChange(UWorld* InWorld, USelection* InActorSelection, USelection* InComponentSelection);
+};
diff --git a/Engine/Plugins/Experimental/FloatingProperties/Source/FloatingProperties/Public/Data/IFloatingPropertiesDataProvider.h b/Engine/Plugins/Experimental/FloatingProperties/Source/FloatingProperties/Public/Data/IFloatingPropertiesDataProvider.h
new file mode 100644
--- /dev/null
+++ b/Engine/Plugins/Experimental/FloatingProperties/Source/FloatingProperties/Public/Data/IFloatingPropertiesDataProvider.h
@@ -0,0 +1,26 @@
+// Copyright Epic Games, Inc. All Rights Reserved.
+
+#pragma once
+
+#include "Containers/Array.h"
+#include "Templates/SharedPointerFwd.h"
+
+class FOnSelectionChanged;
+class IFloatingPropertiesWidgetContainer;
+class USelection;
+class UWorld;
+
+class IFloatingPropertiesDataProvider
+{
+public:
+	/** Must provide consistent shared refs. Same container, same ref. */
+	virtual TArray<TSharedRef<IFloatingPropertiesWidgetContainer>> GetWidgetContainers() = 0;
+
+	virtual USelection* GetActorSelection() const = 0;
+
+	virtual USelection* GetComponentSelection() const = 0;
+
+	virtual UWorld* GetWorld() const = 0;
+
+	virtual bool IsWidgetVisibleInContainer(TSharedRef<IFloatingPropertiesWidgetContainer> InContainer) const = 0;
+};
diff --git a/Engine/Plugins/Experimental/FloatingProperties/Source/FloatingProperties/Public/Data/IFloatingPropertiesWidgetContainer.h b/Engine/Plugins/Experimental/FloatingProperties/Source/FloatingProperties/Public/Data/IFloatingPropertiesWidgetContainer.h
new file mode 100644
--- /dev/null
+++ b/Engine/Plugins/Experimental/FloatingProperties/Source/FloatingProperties/Public/Data/IFloatingPropertiesWidgetContainer.h
@@ -0,0 +1,23 @@
+// Copyright Epic Games, Inc. All Rights Reserved.
+
+#pragma once
+
+#include "Templates/SharedPointer.h"
+
+class SFloatingPropertiesViewportWidget;
+
+class IFloatingPropertiesWidgetContainer
+{
+public:
+	/** Returns the added widget. */
+	virtual TSharedPtr<SFloatingPropertiesViewportWidget> GetWidget() const { return nullptr; }
+
+	/** Adds the widget to teh container, which should store it for later. */
+	virtual bool SetWidget(TSharedRef<SFloatingPropertiesViewportWidget> InWidget) = 0;
+
+	/** Removes the added widget. */
+	virtual bool RemoveWidget() = 0;
+
+	/** Returns true if the widget's container is still valid. */
+	virtual bool IsValid() const = 0;
+};
diff --git a/Engine/Plugins/Experimental/FloatingProperties/Source/FloatingProperties/Public/FloatingPropertiesModule.h b/Engine/Plugins/Experimental/FloatingProperties/Source/FloatingProperties/Public/FloatingPropertiesModule.h
new file mode 100644
--- /dev/null
+++ b/Engine/Plugins/Experimental/FloatingProperties/Source/FloatingProperties/Public/FloatingPropertiesModule.h
@@ -0,0 +1,55 @@
+// Copyright Epic Games, Inc. All Rights Reserved.
+
+#pragma once
+
+#include "Modules/ModuleInterface.h"
+#include "Containers/Map.h"
+#include "Delegates/Delegate.h"
+#include "Delegates/DelegateCombinations.h"
+#include "Logging/LogMacros.h"
+#include "Templates/SharedPointer.h"
+
+class FFloatingPropertiesLevelEditorWidgetController;
+class ILevelEditor;
+class IPropertyHandle;
+class SWidget;
+class UFloatingPropertiesSettings;
+class UScriptStruct;
+
+DECLARE_LOG_CATEGORY_EXTERN(LogFloatingProperties, Log, All);
+
+/**
+ * Floating Properties - Adds floating details panel properties on to the viewport.
+ */
+class FFloatingPropertiesModule : public IModuleInterface
+{
+public:
+	static FFloatingPropertiesModule& Get();
+
+	//~ Begin IModuleInterface
+	virtual void StartupModule() override;
+	virtual void ShutdownModule() override;
+	//~ End IModuleInterface
+
+	DECLARE_DELEGATE_RetVal_OneParam(TSharedPtr<SWidget>, FCreateStructPropertyValueWidgetDelegate, TSharedRef<IPropertyHandle>);
+
+	void RegiserStructPropertyValueWidgetDelegate(UScriptStruct* InStruct, FCreateStructPropertyValueWidgetDelegate InDelegate);
+
+	const FCreateStructPropertyValueWidgetDelegate* GetStructPropertyValueWidgetDelegate(UScriptStruct* InStruct) const;
+
+	void UnregiserStructPropertyValueWidgetDelegate(UScriptStruct* InStruct);
+
+protected:
+	TSharedPtr<FFloatingPropertiesLevelEditorWidgetController> LevelEditorWidgetController;
+	TMap<FName, FCreateStructPropertyValueWidgetDelegate> PropertyValueWidgetDelegates;
+
+	void OnLevelEditorCreated(TSharedPtr<ILevelEditor> InLevelEditor);
+
+	void CreateLevelEditorWidgetController(TSharedPtr<ILevelEditor> InLevelEditor);
+
+	void DestroyLevelEditorWidgetController();
+
+	void OnSettingsChanged(const UFloatingPropertiesSettings* InSettings, FName InSetting);
+
+	void AddDefaultStructPropertyValueWidgetDelegates();
+};
diff --git a/Engine/Plugins/Experimental/FloatingProperties/Source/FloatingProperties/Public/FloatingPropertiesSettings.h b/Engine/Plugins/Experimental/FloatingProperties/Source/FloatingProperties/Public/FloatingPropertiesSettings.h
new file mode 100644
--- /dev/null
+++ b/Engine/Plugins/Experimental/FloatingProperties/Source/FloatingProperties/Public/FloatingPropertiesSettings.h
@@ -0,0 +1,119 @@
+// Copyright Epic Games, Inc. All Rights Reserved.
+
+#pragma once
+
+#include "Engine/DeveloperSettings.h"
+#include "Containers/Map.h"
+#include "Containers/UnrealString.h"
+#include "Math/IntPoint.h"
+#include "Types/SlateEnums.h"
+#include "UObject/Object.h"
+#include "UObject/SoftObjectPtr.h"
+#include "FloatingPropertiesSettings.generated.h"
+
+enum EHorizontalAlignment : int;
+enum EVerticalAlignment : int;
+
+USTRUCT(BlueprintType)
+struct FFloatingPropertiesClassProperty
+{
+	GENERATED_BODY()
+
+	UPROPERTY(EditAnywhere, Category = "Floating Properties")
+	TSoftClassPtr<UObject> Class;
+
+	UPROPERTY(EditAnywhere, Category = "Floating Properties")
+	FString PropertyPath;
+
+	bool operator==(const FFloatingPropertiesClassProperty& InOther) const
+	{
+		return Class == InOther.Class && PropertyPath.Equals(InOther.PropertyPath);
+	}
+
+	bool IsValid() const
+	{
+		return Class.IsValid() && !PropertyPath.IsEmpty();
+	}
+
+	void Reset()
+	{
+		Class = nullptr;
+		PropertyPath = "";
+	}
+
+	friend uint32 GetTypeHash(const FFloatingPropertiesClassProperty& InClassProperty)
+	{
+		return HashCombineFast(GetTypeHash(InClassProperty.Class), GetTypeHash(InClassProperty.PropertyPath));
+	}
+};
+
+USTRUCT(BlueprintType)
+struct FFloatingPropertiesClassProperties
+{
+	GENERATED_BODY()
+
+	UPROPERTY(EditAnywhere, Category = "Floating Properties")
+	TMap<FString, FString> Properties;
+};
+
+/** The Fill alignments are used as a "not set" value. */
+USTRUCT(BlueprintType)
+struct FFloatingPropertiesClassPropertyPosition
+{
+	GENERATED_BODY()
+
+	static const FFloatingPropertiesClassPropertyPosition DefaultStackPosition;
+
+	UPROPERTY(EditAnywhere, Category = "Floating Properties")
+	TEnumAsByte<EHorizontalAlignment> HorizontalAnchor = EHorizontalAlignment::HAlign_Fill;
+
+	UPROPERTY(EditAnywhere, Category = "Floating Properties")
+	TEnumAsByte<EVerticalAlignment> VerticalAnchor = EVerticalAlignment::VAlign_Fill;
+
+	UPROPERTY(EditAnywhere, Category = "Floating Properties")
+	FIntPoint Offset = FIntPoint::ZeroValue;
+
+	bool OnDefaultStack() const;
+};
+
+USTRUCT(BlueprintType)
+struct FFloatingPropertiesClassPropertyAnchor
+{
+	GENERATED_BODY()
+
+	UPROPERTY(EditAnywhere, Category = "Floating Properties")
+	FFloatingPropertiesClassProperty ParentProperty;
+
+	UPROPERTY(EditAnywhere, Category = "Floating Properties")
+	FFloatingPropertiesClassProperty ChildProperty;
+};
+
+UCLASS(Config = EditorPerProjectUserSettings, meta = (DisplayName = "Floating Properties"))
+class FLOATINGPROPERTIES_API UFloatingPropertiesSettings : public UDeveloperSettings
+{
+	GENERATED_BODY()
+
+public:
+	DECLARE_MULTICAST_DELEGATE_TwoParams(FOnFloatingPropertiesSettingsChanged, const UFloatingPropertiesSettings* /* This */, FName /* Setting name */)
+	static FOnFloatingPropertiesSettingsChanged OnChange;
+
+	static FString FindUniqueName(const FFloatingPropertiesClassProperties& InPropertyList, const FString& InDefault);
+
+	UFloatingPropertiesSettings();
+
+	UPROPERTY(Config, EditAnywhere, BlueprintReadWrite, Category = "Floating Properties")
+	bool bEnabled;
+
+	UPROPERTY(Config, EditAnywhere, BlueprintReadWrite, Category = "Floating Properties")
+	TMap<FFloatingPropertiesClassProperty, FFloatingPropertiesClassProperties> SavedValues;
+
+	UPROPERTY(Config, EditAnywhere, BlueprintReadWrite, Category = "Floating Properties")
+	TMap<FFloatingPropertiesClassProperty, FFloatingPropertiesClassPropertyPosition> PropertyPositions;
+
+	UPROPERTY(Config, EditAnywhere, BlueprintReadWrite, Category = "Floating Properties")
+	TMap<FFloatingPropertiesClassProperty, FFloatingPropertiesClassPropertyAnchor> PropertyAnchors;
+
+	//~ Begin UObject
+	virtual void PostEditChangeProperty(FPropertyChangedEvent& PropertyChangedEvent) override;
+	//~ End UObject
+};
