Floating Properties
- Cleanup of property handles and details panel objects on engine pre-exit.

diff --git a/Engine/Plugins/Experimental/FloatingProperties/Source/FloatingProperties/Private/FloatingPropertiesModule.cpp b/Engine/Plugins/Experimental/FloatingProperties/Source/FloatingProperties/Private/FloatingPropertiesModule.cpp
--- a/Engine/Plugins/Experimental/FloatingProperties/Source/FloatingProperties/Private/FloatingPropertiesModule.cpp
+++ b/Engine/Plugins/Experimental/FloatingProperties/Source/FloatingProperties/Private/FloatingPropertiesModule.cpp
@@ -5,6 +5,7 @@
 #include "FloatingPropertiesSettings.h"
 #include "LevelEditor.h"
 #include "LevelEditor/FloatingPropertiesLevelEditorWidgetController.h"
+#include "Misc/CoreDelegates.h"
 #include "Modules/ModuleManager.h"
 #include "Templates/SharedPointer.h"
 #include "Widgets/CustomPropertyWidgets/SFloatingPropertiesColorPropertyEditor.h"
@@ -27,6 +28,8 @@ void FFloatingPropertiesModule::StartupModule()
 	LevelEditorModule.OnLevelEditorCreated().AddRaw(this, &FFloatingPropertiesModule::OnLevelEditorCreated);
 
 	AddDefaultStructPropertyValueWidgetDelegates();
+
+	FCoreDelegates::OnEnginePreExit.AddRaw(this, &FFloatingPropertiesModule::OnEnginePreExit);
 }
 
 void FFloatingPropertiesModule::ShutdownModule()
@@ -41,6 +44,8 @@ void FFloatingPropertiesModule::ShutdownModule()
 	DestroyLevelEditorWidgetController();
 
 	UFloatingPropertiesSettings::OnChange.RemoveAll(this);
+
+	FCoreDelegates::OnEnginePreExit.RemoveAll(this);
 }
 
 void FFloatingPropertiesModule::RegiserStructPropertyValueWidgetDelegate(UScriptStruct* InStruct,
@@ -162,4 +167,10 @@ void FFloatingPropertiesModule::AddDefaultStructPropertyValueWidgetDelegates()
 	);
 }
 
+void FFloatingPropertiesModule::OnEnginePreExit()
+{
+	PropertyValueWidgetDelegates.Empty();
+	LevelEditorWidgetController.Reset();
+}
+
 IMPLEMENT_MODULE(FFloatingPropertiesModule, FloatingProperties)
diff --git a/Engine/Plugins/Experimental/FloatingProperties/Source/FloatingProperties/Public/Data/FloatingPropertiesWidgetController.h b/Engine/Plugins/Experimental/FloatingProperties/Source/FloatingProperties/Public/Data/FloatingPropertiesWidgetController.h
--- a/Engine/Plugins/Experimental/FloatingProperties/Source/FloatingProperties/Public/Data/FloatingPropertiesWidgetController.h
+++ b/Engine/Plugins/Experimental/FloatingProperties/Source/FloatingProperties/Public/Data/FloatingPropertiesWidgetController.h
@@ -24,12 +24,24 @@ struct FFloatingPropertiesPropertyTypes
 	TSharedPtr<IPropertyHandle> PropertyHandle;
 	FProperty* Property;
 	TSharedPtr<IDetailTreeNode> TreeNode;
+
+	~FFloatingPropertiesPropertyTypes()
+	{
+		PropertyHandle.Reset();
+		TreeNode.Reset();
+	}
 };
 
 struct FFloatingPropertiesPropertyList
 {
 	TSharedPtr<IPropertyRowGenerator> Generator;
 	TArray<FFloatingPropertiesPropertyTypes> Properties;
+
+	~FFloatingPropertiesPropertyList()
+	{
+		Properties.Empty();
+		Generator.Reset();
+	}
 };
 
 DECLARE_DELEGATE_RetVal(bool, FIsVisibleDelegate);
diff --git a/Engine/Plugins/Experimental/FloatingProperties/Source/FloatingProperties/Public/FloatingPropertiesModule.h b/Engine/Plugins/Experimental/FloatingProperties/Source/FloatingProperties/Public/FloatingPropertiesModule.h
--- a/Engine/Plugins/Experimental/FloatingProperties/Source/FloatingProperties/Public/FloatingPropertiesModule.h
+++ b/Engine/Plugins/Experimental/FloatingProperties/Source/FloatingProperties/Public/FloatingPropertiesModule.h
@@ -52,4 +52,6 @@ protected:
 	void OnSettingsChanged(const UFloatingPropertiesSettings* InSettings, FName InSetting);
 
 	void AddDefaultStructPropertyValueWidgetDelegates();
+
+	void OnEnginePreExit();
 };
