Floating Properties
- Now properly accounts for DPI.

diff --git a/Engine/Plugins/Experimental/FloatingProperties/Source/FloatingProperties/Private/Data/FloatingPropertiesPropertyNode.cpp b/Engine/Plugins/Experimental/FloatingProperties/Source/FloatingProperties/Private/Data/FloatingPropertiesPropertyNode.cpp
--- a/Engine/Plugins/Experimental/FloatingProperties/Source/FloatingProperties/Private/Data/FloatingPropertiesPropertyNode.cpp
+++ b/Engine/Plugins/Experimental/FloatingProperties/Source/FloatingProperties/Private/Data/FloatingPropertiesPropertyNode.cpp
@@ -98,7 +98,7 @@ TSharedPtr<FFloatingPropertiesPropertyNode> FFloatingPropertiesPropertyNode::Rem
 
 	Child->PropertyPosition = CalculatePropertyPosition(
 		LinkedList->GetDraggableArea(), 
-		Child->GetPropertyWidget()->GetDesiredSize(),
+		Child->GetPropertyWidget()->GetAbsolutePropertySize(),
 		Child->GetCachedPosition()
 	);
 
@@ -409,10 +409,9 @@ FVector2f FFloatingPropertiesPropertyNode::GetStackSize() const
 
 	for (const TSharedRef<FFloatingPropertiesPropertyNode>& Node : GetNodeStack())
 	{
-		const FVector2f DesiredSize = Node->GetPropertyWidget()->GetDesiredSize();
-
-		StackSize.X = FMath::Max(StackSize.X, DesiredSize.X);
-		StackSize.Y += DesiredSize.Y;
+		const FVector2f NodeWidgetSize = Node->GetPropertyWidget()->GetAbsolutePropertySize();
+		StackSize.X = FMath::Max(StackSize.X, NodeWidgetSize.X);
+		StackSize.Y += NodeWidgetSize.Y;
 	}
 
 	return StackSize;
@@ -445,7 +444,7 @@ EFloatingPropertiesUpdateResult FFloatingPropertiesPropertyNode::UpdatePropertyN
 		}
 
 		const FVector2f& PreviousNodePosition = ParentNode->CachedPosition;
-		const FVector2f PreviousNodeSize = ParentNode->PropertyWidget->GetDesiredSize();
+		const FVector2f PreviousNodeSize = ParentNode->PropertyWidget->GetAbsolutePropertySize();
 
 		SetCachedPosition({
 			PreviousNodePosition.X,
diff --git a/Engine/Plugins/Experimental/FloatingProperties/Source/FloatingProperties/Private/Data/FloatingPropertiesSnapMetrics.cpp b/Engine/Plugins/Experimental/FloatingProperties/Source/FloatingProperties/Private/Data/FloatingPropertiesSnapMetrics.cpp
--- a/Engine/Plugins/Experimental/FloatingProperties/Source/FloatingProperties/Private/Data/FloatingPropertiesSnapMetrics.cpp
+++ b/Engine/Plugins/Experimental/FloatingProperties/Source/FloatingProperties/Private/Data/FloatingPropertiesSnapMetrics.cpp
@@ -25,7 +25,7 @@ FVector2f FFloatingPropertiesSnapMetrics::GetSnapPosition(TSharedRef<FFloatingPr
 
 	if (InAttachType == EFloatingPropertiesSnapType::AttachAsParent)
 	{
-		Position.Y += InNode->GetPropertyWidget()->GetDesiredSize().Y;
+		Position.Y += InNode->GetPropertyWidget()->GetAbsolutePropertySize().Y;
 	}
 
 	return Position;
diff --git a/Engine/Plugins/Experimental/FloatingProperties/Source/FloatingProperties/Private/Widgets/SFloatingPropertiesPropertyWidget.cpp b/Engine/Plugins/Experimental/FloatingProperties/Source/FloatingProperties/Private/Widgets/SFloatingPropertiesPropertyWidget.cpp
--- a/Engine/Plugins/Experimental/FloatingProperties/Source/FloatingProperties/Private/Widgets/SFloatingPropertiesPropertyWidget.cpp
+++ b/Engine/Plugins/Experimental/FloatingProperties/Source/FloatingProperties/Private/Widgets/SFloatingPropertiesPropertyWidget.cpp
@@ -259,6 +259,19 @@ FFloatingPropertiesClassProperties* SFloatingPropertiesPropertyWidget::GetSavedV
 	return nullptr;
 }
 
+FVector2f SFloatingPropertiesPropertyWidget::GetAbsolutePropertySize() const
+{
+	const FGeometry& NodePropertyWidgetGeometry = GetTickSpaceGeometry();
+	const FVector2D LocalSize = NodePropertyWidgetGeometry.GetLocalSize();
+
+	if (LocalSize.IsNearlyZero())
+	{
+		return FVector2f::ZeroVector;
+	}
+
+	return GetDesiredSize() * NodePropertyWidgetGeometry.GetAbsoluteSize() / LocalSize;
+}
+
 void SFloatingPropertiesPropertyWidget::SaveConfig()
 {
 	if (UFloatingPropertiesSettings* FloatingPropertiesSettings = GetMutableDefault<UFloatingPropertiesSettings>())
diff --git a/Engine/Plugins/Experimental/FloatingProperties/Source/FloatingProperties/Private/Widgets/SFloatingPropertiesPropertyWidget.h b/Engine/Plugins/Experimental/FloatingProperties/Source/FloatingProperties/Private/Widgets/SFloatingPropertiesPropertyWidget.h
--- a/Engine/Plugins/Experimental/FloatingProperties/Source/FloatingProperties/Private/Widgets/SFloatingPropertiesPropertyWidget.h
+++ b/Engine/Plugins/Experimental/FloatingProperties/Source/FloatingProperties/Private/Widgets/SFloatingPropertiesPropertyWidget.h
@@ -41,6 +41,8 @@ public:
 
 	void SetPropertyValueOverrideSize(FOptionalSize InSize) { PropertyValueOverrideSize = InSize; }
 
+	FVector2f GetAbsolutePropertySize() const;
+
 	void SaveConfig();
 
 protected:
diff --git a/Engine/Plugins/Experimental/FloatingProperties/Source/FloatingProperties/Private/Widgets/SFloatingPropertiesViewportWidget.cpp b/Engine/Plugins/Experimental/FloatingProperties/Source/FloatingProperties/Private/Widgets/SFloatingPropertiesViewportWidget.cpp
--- a/Engine/Plugins/Experimental/FloatingProperties/Source/FloatingProperties/Private/Widgets/SFloatingPropertiesViewportWidget.cpp
+++ b/Engine/Plugins/Experimental/FloatingProperties/Source/FloatingProperties/Private/Widgets/SFloatingPropertiesViewportWidget.cpp
@@ -170,27 +170,11 @@ void SFloatingPropertiesViewportWidget::BuildWidget(AActor* InSelectedActor, con
 		ClassPropertyToNode.Add(ClassProperty, Node);
 
 		FieldContainer->AddSlot()
-		.Size_Lambda([WidgetWeak = Node->GetPropertyWidget().ToWeakPtr()]() -> FVector2D
-			{
-				if (TSharedPtr<SFloatingPropertiesPropertyWidget> Widget = WidgetWeak.Pin())
-				{
-					return Widget->GetDesiredSize();
-				}
-
-				return FVector2D::ZeroVector;
-			})
-		.Position_Lambda([NodeWeak = Node.ToWeakPtr()]() -> FVector2D
-			{
-				if (TSharedPtr<FFloatingPropertiesPropertyNode> Node = NodeWeak.Pin())
-				{
-					return static_cast<FVector2D>(Node->GetCachedPosition());
-				}
-
-				return FVector2D::ZeroVector;
-			})
-		[
-			Node->GetPropertyWidget()
-		];
+			.Size(this, &SFloatingPropertiesViewportWidget::GetNodeSize, Node->GetPropertyWidget().ToWeakPtr())
+			.Position(this, &SFloatingPropertiesViewportWidget::GetNodePosition, Node.ToWeakPtr())
+			[
+				Node->GetPropertyWidget()
+			];
 	}
 
 	// Load data
@@ -269,7 +253,9 @@ FVector2f SFloatingPropertiesViewportWidget::GetDraggableArea() const
 {
 	static const FVector2f BorderArea = FVector2f(4.f, 4.f);
 
-	return GetTickSpaceGeometry().GetLocalSize() - BorderArea;
+	const FGeometry& Geometry = GetTickSpaceGeometry();
+
+	return Geometry.GetAbsoluteSize() - BorderArea;
 }
 
 void SFloatingPropertiesViewportWidget::OnArrangeChildren(const FGeometry& AllottedGeometry, FArrangedChildren& ArrangedChildren) const
@@ -527,4 +513,32 @@ void SFloatingPropertiesViewportWidget::EnsurePositions(bool bInInvalidate) cons
 	}
 }
 
+FVector2D SFloatingPropertiesViewportWidget::GetNodeSize(TWeakPtr<SFloatingPropertiesPropertyWidget> InNodeWidgetWeak) const
+{
+	if (TSharedPtr<SFloatingPropertiesPropertyWidget> Widget = InNodeWidgetWeak.Pin())
+	{
+		return Widget->GetDesiredSize();
+	}
+
+	return FVector2D::ZeroVector;
+}
+
+FVector2D SFloatingPropertiesViewportWidget::GetNodePosition(TWeakPtr<FFloatingPropertiesPropertyNode> InNodeWeak) const
+{
+	if (TSharedPtr<FFloatingPropertiesPropertyNode> Node = InNodeWeak.Pin())
+	{
+		const FGeometry& TickSpaceGeometry = GetTickSpaceGeometry();
+		const FVector2f LocalSize = TickSpaceGeometry.GetAbsoluteSize();
+
+		if (LocalSize.IsNearlyZero())
+		{
+			return FVector2D::ZeroVector;
+		}
+
+		return static_cast<FVector2D>(Node->GetCachedPosition() * TickSpaceGeometry.GetLocalSize() / LocalSize);
+	}
+
+	return FVector2D::ZeroVector;
+}
+
 #undef LOCTEXT_NAMESPACE
diff --git a/Engine/Plugins/Experimental/FloatingProperties/Source/FloatingProperties/Private/Widgets/SFloatingPropertiesViewportWidget.h b/Engine/Plugins/Experimental/FloatingProperties/Source/FloatingProperties/Private/Widgets/SFloatingPropertiesViewportWidget.h
--- a/Engine/Plugins/Experimental/FloatingProperties/Source/FloatingProperties/Private/Widgets/SFloatingPropertiesViewportWidget.h
+++ b/Engine/Plugins/Experimental/FloatingProperties/Source/FloatingProperties/Private/Widgets/SFloatingPropertiesViewportWidget.h
@@ -69,4 +69,8 @@ protected:
 	void CheckForViewportSizeChange(const FVector2f& InViewportSize) const;
 
 	void EnsurePositions(bool bInInvalidate) const;
+
+	FVector2D GetNodeSize(TWeakPtr<SFloatingPropertiesPropertyWidget> InNodeWidgetWeak) const;
+
+	FVector2D GetNodePosition(TWeakPtr<FFloatingPropertiesPropertyNode> InNodeWeak) const;
 };
