Floating Properties: Properties attached to other properties will no longer mistakenly appear at the top-left.

diff --git a/Engine/Plugins/Experimental/FloatingProperties/Source/FloatingProperties/Private/Widgets/SFloatingPropertiesViewportWidget.cpp b/Engine/Plugins/Experimental/FloatingProperties/Source/FloatingProperties/Private/Widgets/SFloatingPropertiesViewportWidget.cpp
--- a/Engine/Plugins/Experimental/FloatingProperties/Source/FloatingProperties/Private/Widgets/SFloatingPropertiesViewportWidget.cpp
+++ b/Engine/Plugins/Experimental/FloatingProperties/Source/FloatingProperties/Private/Widgets/SFloatingPropertiesViewportWidget.cpp
@@ -265,8 +265,11 @@ void SFloatingPropertiesViewportWidget::OnArrangeChildren(const FGeometry& Allot
 	if (!FMath::IsNearlyZero(GetTickSpaceGeometry().GetLocalSize().X))
 	{
 		CheckForViewportSizeChange(AllottedGeometry.GetLocalSize());
-		EnsurePositions(/* bInInvalidate */ !bInitialPositionsSet);
-		bInitialPositionsSet = true;
+
+		if (EnsurePositions(/* bInInvalidate */ !bInitialPositionsSet))
+		{
+			bInitialPositionsSet = true;
+		}
 	}
 }
 
@@ -497,7 +500,7 @@ void SFloatingPropertiesViewportWidget::CheckForViewportSizeChange(const FVector
 	}
 }
 
-void SFloatingPropertiesViewportWidget::EnsurePositions(bool bInInvalidate) const
+bool SFloatingPropertiesViewportWidget::EnsurePositions(bool bInInvalidate) const
 {
 	using namespace UE::FloatingProperties::Private;
 
@@ -505,12 +508,22 @@ void SFloatingPropertiesViewportWidget::EnsurePositions(bool bInInvalidate) cons
 	{
 		if (!Node->HasParent())
 		{
-			if (Node->UpdateStackPositions(bInInvalidate) == EFloatingPropertiesUpdateResult::Updated)
+			switch (Node->UpdateStackPositions(bInInvalidate))
 			{
-				HomogenizeStackWidth(Node);
+				case EFloatingPropertiesUpdateResult::Failed:
+					return false;
+
+				case EFloatingPropertiesUpdateResult::AlreadyUpToDate:
+					break;
+
+				case EFloatingPropertiesUpdateResult::Updated:
+					HomogenizeStackWidth(Node);
+					break;
 			}
 		}
 	}
+
+	return true;
 }
 
 FVector2D SFloatingPropertiesViewportWidget::GetNodeSize(TWeakPtr<SFloatingPropertiesPropertyWidget> InNodeWidgetWeak) const
diff --git a/Engine/Plugins/Experimental/FloatingProperties/Source/FloatingProperties/Private/Widgets/SFloatingPropertiesViewportWidget.h b/Engine/Plugins/Experimental/FloatingProperties/Source/FloatingProperties/Private/Widgets/SFloatingPropertiesViewportWidget.h
--- a/Engine/Plugins/Experimental/FloatingProperties/Source/FloatingProperties/Private/Widgets/SFloatingPropertiesViewportWidget.h
+++ b/Engine/Plugins/Experimental/FloatingProperties/Source/FloatingProperties/Private/Widgets/SFloatingPropertiesViewportWidget.h
@@ -68,7 +68,7 @@ protected:
 
 	void CheckForViewportSizeChange(const FVector2f& InViewportSize) const;
 
-	void EnsurePositions(bool bInInvalidate) const;
+	bool EnsurePositions(bool bInInvalidate) const;
 
 	FVector2D GetNodeSize(TWeakPtr<SFloatingPropertiesPropertyWidget> InNodeWidgetWeak) const;
 
